Add an Insights plugin which analyzes and visualizes data produced by render tracing.

The initial implementation handles command list recording, translation and submission. It adds a number of tracks to show RDG pass execution, RHI command list instances, RHI translate jobs, submission queue processing and interrupt processing. It shows relationships between these stages with arrows, e.g. from an RDG pass to all the command lists it used, or from a command list to the translate job which executed it.

This is primarily intended as a debugging tool, especially if we can have this channel enabled by default, and attach the Insights tail buffer to crash reports. The plugin can show what the renderer was doing at the time the crash occured, which can be useful in understanding hangs and other types of difficult to debug problems. Other types of events will be added in the future with the same goal of providing more context for crashes.

#rnx
#Synced-CL 51619993

diff --git a/Engine/Plugins/RenderTraceInsights/RenderTraceInsights.uplugin b/Engine/Plugins/RenderTraceInsights/RenderTraceInsights.uplugin
new file mode 100644
--- /dev/null
+++ b/Engine/Plugins/RenderTraceInsights/RenderTraceInsights.uplugin
@@ -0,0 +1,28 @@
+{
+	"FileVersion": 3,
+	"Version": 1,
+	"VersionName": "1.0",
+	"FriendlyName": "RenderTrace Insights",
+	"Description": "Rendering debugging facilities in Unreal Insights",
+	"Category": "Insights",
+	"CreatedBy": "Epic Games, Inc.",
+	"CreatedByURL": "https://epicgames.com",
+	"DocsURL": "",
+	"MarketplaceURL": "",
+	"SupportURL": "",
+	"EnabledByDefault": true,
+	"CanContainContent": false,
+	"IsBetaVersion": false,
+	"Installed": false,
+	"SupportedPrograms": [ "UnrealInsights" ],
+	"Modules": [
+		{
+			"Name": "RenderTraceInsights",
+			"Type": "EditorAndProgram",
+			"LoadingPhase": "Default",
+			"ProgramAllowList": [
+				"UnrealInsights"
+			]
+		}
+	]
+}
diff --git a/Engine/Plugins/RenderTraceInsights/Source/RenderTraceInsights/Private/RenderTraceAnalyzer.cpp b/Engine/Plugins/RenderTraceInsights/Source/RenderTraceInsights/Private/RenderTraceAnalyzer.cpp
new file mode 100644
--- /dev/null
+++ b/Engine/Plugins/RenderTraceInsights/Source/RenderTraceInsights/Private/RenderTraceAnalyzer.cpp
@@ -0,0 +1,1618 @@
+// Copyright Epic Games, Inc. All Rights Reserved.
+
+#include "RenderTraceAnalyzer.h"
+#include "IRenderTraceProvider.h"
+#include "RenderTraceModule.h"
+#include "TraceServices/Model/Threads.h"
+#include "InsightsCore/Common/TimeUtils.h"
+#include "Common/ProviderLock.h"
+
+namespace UE
+{
+namespace RenderTraceInsights
+{
+
+enum ED3D12CommandListOpType
+{
+	D3D12CommandListOp_Open = 0,
+	D3D12CommandListOp_Close = 1
+};
+
+enum ED3D12FenceOpType
+{
+	D3D12FenceOp_Wait = 0,
+	D3D12FenceOp_Signal = 1
+};
+
+FRenderTraceAnalyzer::FRenderTraceAnalyzer(TraceServices::IAnalysisSession& InSession, IEditableRenderTraceProvider& InProvider)
+	: Session(InSession)
+	, Provider(InProvider)
+{}
+
+void FRenderTraceAnalyzer::OnAnalysisBegin(const FOnAnalysisContext& Context)
+{
+	UE::Trace::IAnalyzer::FInterfaceBuilder& Builder = Context.InterfaceBuilder;
+
+#define ROUTE_EVENT(EventName) Builder.RouteEvent(Packet_##EventName, "RenderTrace", #EventName "Message");
+
+	ROUTE_EVENT(CommandListCreated);
+	ROUTE_EVENT(CommandListDetach);
+	ROUTE_EVENT(CommandListFinishRecording);
+	ROUTE_EVENT(CommandListDestroyed);
+	ROUTE_EVENT(BeginExecuteRDGPass);
+	ROUTE_EVENT(EndExecuteRDGPass);
+	ROUTE_EVENT(BeginExecuteDispatchPass);
+	ROUTE_EVENT(AddDispatchPassCommandList);
+	ROUTE_EVENT(EndExecuteDispatchPass);
+	ROUTE_EVENT(SubmitCommandLists);
+	ROUTE_EVENT(SplitTranslateJob);
+	ROUTE_EVENT(DispatchCommandList);
+	ROUTE_EVENT(BeginTranslateCommandList);
+	ROUTE_EVENT(ActivatePipelines);
+	ROUTE_EVENT(SetTranslateContext);
+	ROUTE_EVENT(EndTranslateCommandList);
+	ROUTE_EVENT(FinalizeTranslateJob);
+	ROUTE_EVENT(SubmitTranslateJobs);
+	ROUTE_EVENT(AssociatePlatformContext);
+	ROUTE_EVENT(D3D12PayloadCreated);
+	ROUTE_EVENT(D3D12AddPayload);
+	ROUTE_EVENT(D3D12CommandListOp);
+	ROUTE_EVENT(D3D12SyncPointOp);
+	ROUTE_EVENT(D3D12ManualFenceOp);
+	ROUTE_EVENT(D3D12ProcessSubmissionQueueEnter);
+	ROUTE_EVENT(D3D12SubmitYieldOnUnresolvedSyncPoint);
+	ROUTE_EVENT(D3D12SubmitYieldOnUnresolvedManualFence);
+	ROUTE_EVENT(D3D12SubmitWaitQueueFence);
+	ROUTE_EVENT(D3D12SubmitWaitManualFence);
+	ROUTE_EVENT(D3D12SubmitPlatformCommandLists);
+	ROUTE_EVENT(D3D12SubmitSignalManualFence);
+	ROUTE_EVENT(D3D12SubmitSignalQueueFence);
+	ROUTE_EVENT(D3D12SubmitResolveSyncPoint);
+	ROUTE_EVENT(D3D12ProcessSubmissionQueueExit);
+	ROUTE_EVENT(D3D12ProcessInterruptQueueEnter);
+	ROUTE_EVENT(D3D12InterruptQueueFenceSignaled);
+	ROUTE_EVENT(D3D12ProcessInterruptQueueExit);
+
+#undef ROUTE_EVENT
+}
+
+static inline FStringView GetString(const UE::Trace::IAnalyzer::FEventData& EventData, const ANSICHAR* FieldName)
+{
+	FStringView Value;
+	EventData.GetString(FieldName, Value);
+	return Value;
+}
+
+void FRenderTraceAnalyzer::BeginCmdListExecute(uint64 CmdListAppID, uint32 PermanentPassID, double Timestamp)
+{
+	FRDGPassInstance& Pass = Provider.EditRDGPass(PermanentPassID);
+
+	const uint32* PermanentCmdListID = LiveCommandLists.Find(CmdListAppID);
+	if (!PermanentCmdListID)
+	{
+		UE_LOG(LogRenderTrace, Warning, TEXT("Got unknown command list %llx for begin execute RDG pass for '%s' (permanent ID %x), ignoring."),
+			CmdListAppID, *Pass.Name, PermanentPassID);
+		return;
+	}
+
+	FCommandListInstance& CmdList = Provider.EditCommandList(*PermanentCmdListID);
+	if (!CmdList.RecordingEvents.IsEmpty() && CmdList.RecordingEvents.Last().End == 0.0)
+	{
+		UE_LOG(LogRenderTrace, Warning, TEXT("Got begin execute RDG pass for '%s' (permanent ID %x) on command list %llx which has an open recording event (ID %x). Ignoring."),
+			*Pass.Name, PermanentPassID, CmdList.AppID, CmdList.RecordingEvents.Last().ItemID);
+		return;
+	}
+
+	Pass.CommandListIDs.Add(*PermanentCmdListID);
+	CmdList.RecordingEvents.Emplace(PermanentPassID, Timestamp);
+}
+
+void FRenderTraceAnalyzer::EndRDGPassExecute(uint32 PermanentPassID, double Timestamp)
+{
+	FRDGPassInstance& Pass = Provider.EditRDGPass(PermanentPassID);
+	Pass.EndTime = Timestamp;
+
+	for (uint32 PermanentCmdListID : Pass.CommandListIDs)
+	{
+		FCommandListInstance& CmdList = Provider.EditCommandList(PermanentCmdListID);
+		check(!CmdList.RecordingEvents.IsEmpty() && CmdList.RecordingEvents.Last().ItemID == PermanentPassID);
+
+		// Sometimes FinishRecording arrives before a pass ends, e.g. for dispatch passes. When that happens it closes the recording
+		// interval, so we should only set the end timestamp if it's not already set.
+		if (CmdList.RecordingEvents.Last().End == 0.0)
+		{
+			CmdList.RecordingEvents.Last().End = Timestamp;
+		}
+	}
+
+	IRenderTraceProvider::TEventTimeline& Timeline = Provider.EditRDGThreadTimeline(Pass.ExecThreadID);
+	Timeline.AppendEndEvent(Timestamp);
+}
+
+void FRenderTraceAnalyzer::AddCommandListToTimeline(uint32 PermanentCmdListID)
+{
+	FCommandListInstance& CmdList = Provider.EditCommandList(PermanentCmdListID);
+
+	// A positive timeline index means the command list was already added, e.g. in FinishRecording and now we're called from
+	// its destroy event. A negative value different from -1 means we've already processed this command list and decided not
+	// to add it to a timeline.
+	if (CmdList.TimelineIndex != -1)
+	{
+		return;
+	}
+
+	int32 FirstEligibleTimeline;
+	double StartTime, EndTime;
+	switch (CmdList.Type)
+	{
+	case ECommandListType::Immediate:
+	{
+		// Try to place immediate command lists on the first timeline. Other command list types will set FirstEligibleTimeline to 1
+		// so they're never placed on the first lane.
+		FirstEligibleTimeline = 0;
+		StartTime = CmdList.CreateTime;
+		EndTime = FMath::Max(CmdList.FinishRecordingTime, CmdList.DestroyTime);
+		break;
+	}
+
+	case ECommandListType::Detached:
+	{
+		// Skip detached command lists if the parent won't be shown.
+		if (CmdList.SourceCmdListID == INVALID_EVENT_ID)
+		{
+			CmdList.TimelineIndex = -2;
+			return;
+		}
+
+		const FCommandListInstance& SourceCmdList = Provider.EditCommandList(CmdList.SourceCmdListID);
+		if (SourceCmdList.RecordingEvents.IsEmpty())
+		{
+			CmdList.TimelineIndex = -2;
+			return;
+		}
+
+		FirstEligibleTimeline = 1;
+		StartTime = CmdList.CreateTime;
+		EndTime = FMath::Max(CmdList.FinishRecordingTime, StartTime);
+		break;
+	}
+
+	case ECommandListType::Regular:
+	{
+		// Skip regular command lists with no recording events.
+		if (CmdList.RecordingEvents.IsEmpty())
+		{
+			CmdList.TimelineIndex = -2;
+			return;
+		}
+
+		FirstEligibleTimeline = 1;
+
+		// Nothing interesting happens before the first recording event, so start the command list event then.
+		StartTime = CmdList.RecordingEvents[0].Start;
+		EndTime = FMath::Max(CmdList.FinishRecordingTime, CmdList.RecordingEvents.Last().End);
+		break;
+	}
+
+	default:
+		checkNoEntry();
+		return;
+	}
+
+	CmdList.TimelineStart = StartTime;
+
+	if (StartTime == EndTime)
+	{
+		// The tiny command lists made for closing out parallel dispatch passes finish recording as soon as they're created,
+		// so nudge their end time a little bit in order to make them show up in the timeline.
+		EndTime += UE::Insights::FTimeValue::Microsecond;
+	}
+
+	int32 BestTimelineIdx = -1;
+	IRenderTraceProvider::TEventTimeline* BestTimeline = nullptr;
+	double BestDist = std::numeric_limits<double>::infinity();
+	const int32 NumTimelines = Provider.GetNumCommandListTimelines();
+	for (int32 TimelineIdx = FirstEligibleTimeline; TimelineIdx < NumTimelines; ++TimelineIdx)
+	{
+		IRenderTraceProvider::TEventTimeline& Timeline = Provider.EditCommandListTimeline(TimelineIdx);
+		double Distance = StartTime - Timeline.GetLastTimestamp();
+		if (Distance >= 0 && Distance < BestDist)
+		{
+			BestDist = Distance;
+			BestTimeline = &Timeline;
+			BestTimelineIdx = TimelineIdx;
+			if (TimelineIdx == 0)
+			{
+				// If we're able to place it on the first track, this is an immediate command list, so stop searching
+				// for better matches.
+				break;
+			}
+		}
+	}
+
+	if (BestTimeline == nullptr)
+	{
+		TPair<uint32, IRenderTraceProvider::TEventTimeline&> NewTimeline = Provider.AddCommandListTimeline();
+		BestTimelineIdx = NewTimeline.Key;
+		// PVS-Studio doesn't understand how references work, we need to disable the warning about taking the address of a variable which goes out of scope.
+		BestTimeline = &NewTimeline.Value; //-V506
+	}
+
+	CmdList.TimelineIndex = BestTimelineIdx;
+	BestTimeline->AppendBeginEvent(StartTime, FRenderTraceEvent(PermanentCmdListID));
+
+	for (int32 RecordEventIdx = 0; RecordEventIdx < CmdList.RecordingEvents.Num(); ++RecordEventIdx)
+	{
+		FEventInterval& RecordEvent = CmdList.RecordingEvents[RecordEventIdx];
+		BestTimeline->AppendBeginEvent(RecordEvent.Start, FRenderTraceEvent(RecordEvent.ItemID, PermanentCmdListID));
+
+		if (RecordEvent.End == 0.0)
+		{
+			// Close the event if it's still open. This can happen if FinishRecording arrives before a pass ends, e.g. for parallel dispatch passes.
+			if (RecordEventIdx == CmdList.RecordingEvents.Num() - 1)
+			{
+				RecordEvent.End = EndTime;
+			}
+			else
+			{
+				UE_LOG(LogRenderTrace, Warning, TEXT("Command list %llx (permanent ID %x) is being added to the timeline with an open recording event that's not the last in the list (index %d, num %d)."),
+					CmdList.AppID, PermanentCmdListID, RecordEventIdx, CmdList.RecordingEvents.Num());
+			}
+		}
+
+		if (RecordEvent.End <= RecordEvent.Start)
+		{
+			// Some command lists (e.g. pass epilogues) finish recording as soon as they're created, so we nudge the
+			// end of the recording event a little bit, but we need to be careful not to overlap the next event.
+			const double NextEventStart = (RecordEventIdx < CmdList.RecordingEvents.Num() - 1) ? CmdList.RecordingEvents[RecordEventIdx + 1].Start : EndTime;
+			RecordEvent.End = FMath::Min(RecordEvent.Start + UE::Insights::FTimeValue::Microsecond, NextEventStart);
+		}
+
+		BestTimeline->AppendEndEvent(RecordEvent.End);
+	}
+
+	BestTimeline->AppendEndEvent(EndTime);
+}
+
+void FRenderTraceAnalyzer::EndCommandList(uint32 PermanentCmdListID, double Timestamp)
+{
+	FCommandListInstance& CmdList = Provider.EditCommandList(PermanentCmdListID);
+	if (!CmdList.RecordingEvents.IsEmpty() && CmdList.RecordingEvents.Last().End == 0.0)
+	{
+		// It's possible to receive a destroy event for the immediate command list while a pass is recording to it, e.g. FRDGDispatchPass::Execute() calls
+		// QueueAsyncCommandListSubmit(), which can detach the pending commands in the immediate command list. We only log a warning if this happens for
+		// other command list types. In both cases we close the pending recording event.
+		if (CmdList.Type != ECommandListType::Immediate)
+		{
+			UE_LOG(LogRenderTrace, Warning, TEXT("Command list %llx (permanent ID %x) ended with open recording event for pass %x, force-closing."),
+				CmdList.AppID, PermanentCmdListID, CmdList.RecordingEvents.Last().ItemID);
+		}
+		CmdList.RecordingEvents.Last().End = Timestamp;
+	}
+	CmdList.DestroyTime = Timestamp;
+
+	AddCommandListToTimeline(PermanentCmdListID);
+}
+
+void FRenderTraceAnalyzer::EndTranslateTask(uint32 PermanentTaskID, double Timestamp)
+{
+	FRHITranslateTask& TranslateTask = Provider.EditRHITranslateTask(PermanentTaskID);
+	check(TranslateTask.Type == ERHITranslateTaskType::Translate);
+
+	if (TranslateTask.EndTime == 0.0)
+	{
+		TranslateTask.EndTime = Timestamp;
+	}
+
+	for (int PipeIdx = 0; PipeIdx < 2; ++PipeIdx)
+	{
+		FRHITranslateContext& TranslateContext = TranslateTask.Contexts[PipeIdx];
+		if (!TranslateContext.ActiveIntervals.IsEmpty() && TranslateContext.ActiveIntervals.Last().End == 0.0)
+		{
+			TranslateContext.ActiveIntervals.Last().End = TranslateTask.EndTime;
+		}
+
+		for (uint32 PayloadID : TranslateContext.PayloadIDs)
+		{
+			FPlatformPayload& Payload = Provider.EditPlatformPayload(PayloadID);
+			if (Payload.EndTime == 0.0)
+			{
+				Payload.EndTime = TranslateTask.EndTime;
+			}
+
+			Payload.StartTime = FMath::Min(Payload.StartTime, TranslateTask.EndTime);
+			Payload.EndTime = FMath::Min(Payload.EndTime, TranslateTask.EndTime);
+		}
+	}
+
+	IRenderTraceProvider::TEventTimeline& Timeline = Provider.EditRHITranslateThreadTimeline(TranslateTask.ThreadID);
+	Timeline.AppendBeginEvent(TranslateTask.StartTime, FRenderTraceEvent(PermanentTaskID));
+
+	// Add subevents for the command lists which are translated.
+	for (const FEventInterval& CmdList : TranslateTask.ProcessedItems)
+	{
+		Timeline.AppendBeginEvent(CmdList.Start, FRenderTraceEvent(CmdList.ItemID, PermanentTaskID));
+		Timeline.AppendEndEvent(CmdList.End);
+	}
+	// Translate tasks have well defined durations, use the timestamp.
+	Timeline.AppendEndEvent(TranslateTask.EndTime);
+}
+
+FPlatformPayload* FRenderTraceAnalyzer::FindPlatformPayload(uint64 PayloadID, const TCHAR* EventName, uint32 ThreadID)
+{
+	const uint32* ExistingPermanentID = LivePlatformPayloads.Find(PayloadID);
+	if (!ExistingPermanentID)
+	{
+		if (ThreadID != GameThreadID)
+		{
+			UE_LOG(LogRenderTrace, Warning, TEXT("Got %s event for payload ID %llx, but the payload is not in the live payloads list."), EventName, PayloadID);
+		}
+		return nullptr;
+	}
+
+	return &Provider.EditPlatformPayload(*ExistingPermanentID);
+}
+
+void FRenderTraceAnalyzer::ProcessCommandListCreated(const FEventData& EventData, uint32 ThreadID, double Timestamp)
+{
+	const uint64 CmdListID = EventData.GetValue<uint64>("CmdListID");
+	const bool bImmediate = EventData.GetValue<bool>("bImmediate");
+
+	const uint32* ExistingPermanentID = LiveCommandLists.Find(CmdListID);
+	if (ExistingPermanentID)
+	{
+		UE_LOG(LogRenderTrace, Warning, TEXT("Got create command list for ID %llx while it's already open, ignoring."), CmdListID);
+		return;
+	}
+
+	TPair<uint32, FCommandListInstance&> NewCmdListEntry = Provider.AddCommandList(CmdListID, Timestamp, bImmediate ? ECommandListType::Immediate : ECommandListType::Regular);
+	LiveCommandLists.Add(CmdListID, NewCmdListEntry.Key);
+}
+
+void FRenderTraceAnalyzer::ProcessCommandListDetach(const FEventData& EventData, uint32 ThreadID, double Timestamp)
+{
+	const uint64 SourceCmdListID = EventData.GetValue<uint64>("SourceCmdListID");
+	const uint64 DetachedCmdListID = EventData.GetValue<uint64>("DetachedCmdListID");
+	const bool bIsImmediateFlush = EventData.GetValue<bool>("bIsImmediateFlush");
+
+	const uint32* PermanentSourceCmdListIDPtr = LiveCommandLists.Find(SourceCmdListID);
+	if (!PermanentSourceCmdListIDPtr)
+	{
+		UE_LOG(LogRenderTrace, Warning, TEXT("Got detach command list event for unknown command list ID %llx."), SourceCmdListID);
+		return;
+	}
+	const uint32 PermanentSourceCmdListID = *PermanentSourceCmdListIDPtr;
+
+	// The new command list is created using a move constructor which doesn't send CommandListCreated events for now. Add it to the live map
+	// if it's not there already.
+	uint32 PermanentDetachedCmdListID;
+	const uint32* ExistingDetachedPermanentID = LiveCommandLists.Find(DetachedCmdListID);
+	if (!ExistingDetachedPermanentID)
+	{
+		TPair<uint32, FCommandListInstance&> NewCmdListEntry = Provider.AddCommandList(DetachedCmdListID, Timestamp, ECommandListType::Detached);
+		PermanentDetachedCmdListID = NewCmdListEntry.Key;
+		LiveCommandLists.Add(DetachedCmdListID, PermanentDetachedCmdListID);
+	}
+	else
+	{
+		PermanentDetachedCmdListID = *ExistingDetachedPermanentID;
+	}
+
+	FCommandListInstance& SourceCmdList = Provider.EditCommandList(PermanentSourceCmdListID);
+	if (SourceCmdList.DetachedCmdListID != INVALID_EVENT_ID)
+	{
+		UE_LOG(LogRenderTrace, Warning, TEXT("Got detach command list event for source command list %llx to new command list %llx, but the source has already been detached to permanent ID %x."),
+			SourceCmdListID, DetachedCmdListID, SourceCmdList.DetachedCmdListID);
+	}
+	SourceCmdList.DetachedCmdListID = PermanentDetachedCmdListID;
+	SourceCmdList.DetachTime = Timestamp;
+
+	FCommandListInstance& DetachedCmdList = Provider.EditCommandList(PermanentDetachedCmdListID);
+	DetachedCmdList.Type = ECommandListType::Detached;
+	DetachedCmdList.SourceCmdListID = PermanentSourceCmdListID;
+}
+
+void FRenderTraceAnalyzer::ProcessCommandListFinishRecording(const FEventData& EventData, uint32 ThreadID, double Timestamp)
+{
+	const uint64 CmdListID = EventData.GetValue<uint64>("CmdListID");
+	const uint8 Flags = EventData.GetValue<uint8>("Flags");
+
+	const uint32* PermanentCmdListID = LiveCommandLists.Find(CmdListID);
+	if (!PermanentCmdListID)
+	{
+		UE_LOG(LogRenderTrace, Warning, TEXT("Got finish recording event for unknown command list ID %llx."), CmdListID);
+		return;
+	}
+
+	FCommandListInstance& CmdList = Provider.EditCommandList(*PermanentCmdListID);
+	CmdList.FinishRecordingTime = Timestamp;
+	CmdList.RecordingFlags = Flags;
+
+	AddCommandListToTimeline(*PermanentCmdListID);
+}
+
+void FRenderTraceAnalyzer::ProcessCommandListDestroyed(const FEventData& EventData, uint32 ThreadID, double Timestamp)
+{
+	const uint64 CmdListID = EventData.GetValue<uint64>("CmdListID");
+	const uint32* PermanentCmdListID = LiveCommandLists.Find(CmdListID);
+	if (!PermanentCmdListID)
+	{
+		UE_LOG(LogRenderTrace, Warning, TEXT("Got destroy event for unknown command list ID %llx."), CmdListID);
+		return;
+	}
+
+	EndCommandList(*PermanentCmdListID, Timestamp);
+	LiveCommandLists.Remove(CmdListID);
+}
+
+void FRenderTraceAnalyzer::ProcessBeginExecuteRDGPass(const FEventData& EventData, uint32 ThreadID, double Timestamp, bool bIsDispatch)
+{
+	const uint64 PassID = EventData.GetValue<uint64>("PassID");
+	FString Name(GetString(EventData, "Name"));
+
+	const uint32* ExistingPermanentID = LiveRDGPasses.Find(PassID);
+	if (ExistingPermanentID)
+	{
+		UE_LOG(LogRenderTrace, Warning, TEXT("Got execute RDG pass for '%s' (app ID %llx) while it's already open (permanent ID %x), ignoring."),
+			*Name, PassID, *ExistingPermanentID);
+		return;
+	}
+
+	TPair<uint32, FRDGPassInstance&> NewPassEntry = Provider.AddRDGPass(Name, ThreadID, Timestamp, bIsDispatch ? ERDGPassType::Dispatch : ERDGPassType::Regular);
+	LiveRDGPasses.Add(PassID, NewPassEntry.Key);
+
+	if (!bIsDispatch)
+	{
+		const uint64 CmdListID = EventData.GetValue<uint64>("CmdListID");
+		BeginCmdListExecute(CmdListID, NewPassEntry.Key, Timestamp);
+
+		NewPassEntry.Value.Flags = EventData.GetValue<uint16>("PassFlags");
+		NewPassEntry.Value.TaskMode = static_cast<ERDGPassTaskMode>(EventData.GetValue<uint8>("TaskMode"));
+	}
+
+	IRenderTraceProvider::TEventTimeline& Timeline = Provider.EditRDGThreadTimeline(ThreadID);
+	Timeline.AppendBeginEvent(Timestamp, FRenderTraceEvent(NewPassEntry.Key));
+}
+
+void FRenderTraceAnalyzer::ProcessEndExecuteRDGPass(const FEventData& EventData, uint32 ThreadID, double Timestamp, bool bIsDispatch)
+{
+	const uint64 PassID = EventData.GetValue<uint64>("PassID");
+	const uint32* PermanentPassID = LiveRDGPasses.Find(PassID);
+	if (!PermanentPassID)
+	{
+		UE_LOG(LogRenderTrace, Warning, TEXT("Got end execute RDG pass for unknown pass ID %llx."), PassID);
+		return;
+	}
+
+	if (bIsDispatch)
+	{
+		const uint64 EpilogueCmdListID = EventData.GetValue<uint64>("EpilogueCmdListID");
+		if (EpilogueCmdListID > 0)
+		{
+			BeginCmdListExecute(EpilogueCmdListID, *PermanentPassID, Timestamp);
+		}
+	}
+
+	EndRDGPassExecute(*PermanentPassID, Timestamp);
+	LiveRDGPasses.Remove(PassID);
+}
+
+void FRenderTraceAnalyzer::ProcessAddDispatchPassCommandList(const FEventData& EventData, uint32 ThreadID, double Timestamp)
+{
+	const uint64 PassID = EventData.GetValue<uint64>("PassID");
+	const uint64 CmdListID = EventData.GetValue<uint64>("CmdListID");
+
+	const uint32* PermanentPassID = LiveRDGPasses.Find(PassID);
+	if (!PermanentPassID)
+	{
+		UE_LOG(LogRenderTrace, Warning, TEXT("Got add RDG pass command list for unknown pass ID %llx."), PassID);
+		return;
+	}
+
+	FRDGPassInstance& Pass = Provider.EditRDGPass(*PermanentPassID);
+	BeginCmdListExecute(CmdListID, *PermanentPassID, Timestamp);
+}
+
+void FRenderTraceAnalyzer::ProcessSubmitCommandLists(const FEventData& EventData, uint32 ThreadID, double Timestamp)
+{
+	TConstArrayView<uint64> CmdListIDs = EventData.GetArrayView<uint64>("CmdListIDs");
+	const uint32 RTSubmissionID = EventData.GetValue<uint32>("RTSubmissionID");
+	const uint16 Flags = EventData.GetValue<uint16>("Flags");
+
+	// We assign an arbitrary duration of 10 us to these events so that they can be seen on the timeline. If a second submission
+	// event arrives less than 10 us after we added the previous one, we need to nudge it slightly to fix overlaps.
+	IRenderTraceProvider::TEventTimeline& Timeline = Provider.EditRenderThreadSubmissionTimeline();
+	const double EventStartTime = FMath::Max(Timestamp, Timeline.GetLastTimestamp());
+
+	// Add a fake RDG pass to represent the submit.
+	const TCHAR* PassName = (RTSubmissionID != 0) ? TEXT("SubmitToGPU") : TEXT("Queue Submit");
+	TPair<uint32, FRDGPassInstance&> SubmitPassEntry = Provider.AddRDGPass(PassName, ThreadID, EventStartTime, ERDGPassType::Submit);
+	FRDGPassInstance& Pass = SubmitPassEntry.Value;
+	Pass.Flags = Flags;
+	Pass.EndTime = Pass.StartTime + 10 * UE::Insights::FTimeValue::Microsecond;
+
+	if (RTSubmissionID != 0)
+	{
+		// Remember which pass ID corresponds to the app-provided submission ID.
+		PendingRTSubmitIDs.Add(RTSubmissionID, SubmitPassEntry.Key);
+	}
+
+	Pass.CommandListIDs.Reserve(CmdListIDs.Num());
+	for (uint64 CmdListID : CmdListIDs)
+	{
+		const uint32* PermanentCmdListID = LiveCommandLists.Find(CmdListID);
+		if (!PermanentCmdListID)
+		{
+			UE_LOG(LogRenderTrace, Warning, TEXT("Got submit event for unknown command list ID %llx."), CmdListID);
+			continue;
+		}
+
+		Pass.CommandListIDs.Add(*PermanentCmdListID);
+		FCommandListInstance& CmdList = Provider.EditCommandList(*PermanentCmdListID);
+		CmdList.SubmitTime = Timestamp;
+		CmdList.SubmitPassID = SubmitPassEntry.Key;
+	}
+
+	Timeline.AppendBeginEvent(Pass.StartTime, FRenderTraceEvent(SubmitPassEntry.Key));
+	Timeline.AppendEndEvent(Pass.EndTime);
+}
+
+void FRenderTraceAnalyzer::ProcessSplitTranslateJob(const FEventData& EventData, uint32 ThreadID, double Timestamp)
+{
+	const uint64 TranslateJobID = EventData.GetValue<uint64>("TranslateJobID");
+	const uint64 CmdListID = EventData.GetValue<uint64>("CmdListID");
+	const uint8 Flags = EventData.GetValue<uint8>("Flags");
+
+	const uint32* PermanentCmdListID = LiveCommandLists.Find(CmdListID);
+	if (!PermanentCmdListID)
+	{
+		UE_LOG(LogRenderTrace, Warning, TEXT("Got split event with unknown command list ID %llx."), CmdListID);
+		return;
+	}
+
+	// We store this information on the command list because the translate task object might not exist yet.
+	// A translate tasks is only created when it starts to translate its first command list, which can be after
+	// the event which caused the previous one to be split.
+	FCommandListInstance& CmdList = Provider.EditCommandList(*PermanentCmdListID);
+	CmdList.TranslatePrevJobAppID = TranslateJobID;
+	CmdList.TranslateSplitFlags = Flags;
+}
+
+void FRenderTraceAnalyzer::ProcessDispatchCommandList(const FEventData& EventData, uint32 ThreadID, double Timestamp)
+{
+	const uint64 CmdListID = EventData.GetValue<uint64>("CmdListID");
+
+	const uint32* PermanentCmdListID = LiveCommandLists.Find(CmdListID);
+	if (!PermanentCmdListID)
+	{
+		UE_LOG(LogRenderTrace, Warning, TEXT("Got dispatch event for unknown command list ID %llx."), CmdListID);
+		return;
+	}
+
+	FCommandListInstance& CmdList = Provider.EditCommandList(*PermanentCmdListID);
+	CmdList.DispatchTime = Timestamp;
+}
+
+void FRenderTraceAnalyzer::ProcessBeginTranslateCommandList(const FEventData& EventData, uint32 ThreadID, double Timestamp)
+{
+	const uint64 TranslateJobID = EventData.GetValue<uint64>("TranslateJobID");
+	const uint64 CmdListID = EventData.GetValue<uint64>("CmdListID");
+
+	const uint32* PermanentCmdListID = LiveCommandLists.Find(CmdListID);
+	if (!PermanentCmdListID)
+	{
+		UE_LOG(LogRenderTrace, Warning, TEXT("Got begin translate event for unknown command list ID %llx."), CmdListID);
+		return;
+	}
+
+	FCommandListInstance& CmdList = Provider.EditCommandList(*PermanentCmdListID);
+	if (CmdList.TranslateTaskID != INVALID_EVENT_ID)
+	{
+		UE_LOG(LogRenderTrace, Warning, TEXT("Got begin translate event for command list ID %llx which is already associated with task %x."), CmdListID, CmdList.TranslateTaskID);
+		return;
+	}
+
+	uint32 PermanentTranslateJobID = INVALID_EVENT_ID, ForceSplitFromJobID = INVALID_EVENT_ID;
+
+	const uint32* PermanentTranslateJobIDPtr = LiveTranslateJobs.Find(TranslateJobID);
+	if (PermanentTranslateJobIDPtr)
+	{
+		PermanentTranslateJobID = *PermanentTranslateJobIDPtr;
+		FRHITranslateTask& TranslateJob = Provider.EditRHITranslateTask(PermanentTranslateJobID);
+		check(TranslateJob.Type == ERHITranslateTaskType::Translate);
+		if (TranslateJob.ThreadID != ThreadID)
+		{
+			// The job just jumped to another thread. This happens if a command list finishes translating before the next one is dispatched,
+			// so that the job spawns another taskgraph task. We will close the current job and start a new one.
+			EndTranslateTask(PermanentTranslateJobID, Timestamp);
+			ForceSplitFromJobID = PermanentTranslateJobID;
+			PermanentTranslateJobID = INVALID_EVENT_ID;
+		}
+	}
+
+	if (PermanentTranslateJobID == INVALID_EVENT_ID)
+	{
+		TPair<uint32, FRHITranslateTask&> NewTaskEntry = Provider.AddRHITranslateTask(TranslateJobID, ERHITranslateTaskType::Translate, ThreadID, Timestamp);
+		PermanentTranslateJobID = NewTaskEntry.Key;
+		LiveTranslateJobs.Add(TranslateJobID, PermanentTranslateJobID);
+	}
+
+	CmdList.TranslateTaskID = PermanentTranslateJobID;
+
+	FRHITranslateTask& TranslateJob = Provider.EditRHITranslateTask(PermanentTranslateJobID);
+	TranslateJob.ProcessedItems.Emplace(*PermanentCmdListID, Timestamp);
+
+	if (ForceSplitFromJobID != INVALID_EVENT_ID)
+	{
+		// We've split the job because it jumped threads, so we need to move the contexts from the old job to the new one.
+		// Note that this doesn't move any payloads, so if the new command list is executed on the active payload we won't
+		// be able to show this on the timeline.
+		const FRHITranslateTask& OldJob = Provider.EditRHITranslateTask(ForceSplitFromJobID);
+		for (int PipeIdx = 0; PipeIdx < 2; ++PipeIdx)
+		{
+			const uint64 RHIContextID = OldJob.Contexts[PipeIdx].RHIContextID;
+			if (RHIContextID != 0)
+			{
+				FRHITranslateContext& Context = TranslateJob.Contexts[PipeIdx];
+				Context.RHIContextID = RHIContextID;
+				Context.ActiveIntervals.Emplace(Timestamp);
+				LiveRHIContexts.Add(RHIContextID, PermanentTranslateJobID);
+			}
+		}
+
+		// Record the split reason on the command list so we can link up the tasks.
+		CmdList.TranslatePrevJobAppID = OldJob.AppID;
+		CmdList.TranslateSplitFlags = ERHITranslateJobSplitFlag_JumpThreads;
+	}
+}
+
+FRHITranslateTask* FRenderTraceAnalyzer::GetTranslateJobForCmdList(uint64 CmdListID, const TCHAR* EventName, uint32 ThreadID, uint32* OutTranslateJobID)
+{
+	const uint32* PermanentCmdListID = LiveCommandLists.Find(CmdListID);
+	if (!PermanentCmdListID)
+	{
+		// The engine calls SwitchPipeline() on the immediate command list from the game thread at startup, to make sure it has an active context.
+		// These calls will not be part of a translate job, but it's fine to ignore them. Only log a warning when this happens later.
+		if (ThreadID != GameThreadID)
+		{
+			UE_LOG(LogRenderTrace, Warning, TEXT("Got %s event for unknown command list ID %llx."), EventName, CmdListID);
+		}
+		return nullptr;
+	}
+
+	FCommandListInstance& CmdList = Provider.EditCommandList(*PermanentCmdListID);
+	if (CmdList.TranslateTaskID == INVALID_EVENT_ID)
+	{
+		if (ThreadID != GameThreadID)
+		{
+			UE_LOG(LogRenderTrace, Warning, TEXT("Got %s event for command list ID %llx which doesn't have a translate job."), EventName, CmdListID);
+		}
+		return nullptr;
+	}
+
+	FRHITranslateTask& TranslateJob = Provider.EditRHITranslateTask(CmdList.TranslateTaskID);
+	check(TranslateJob.Type == ERHITranslateTaskType::Translate);
+
+	if (OutTranslateJobID)
+	{
+		*OutTranslateJobID = CmdList.TranslateTaskID;
+	}
+
+	return &TranslateJob;
+}
+
+void FRenderTraceAnalyzer::ProcessActivatePipelines(const FEventData& EventData, uint32 ThreadID, double Timestamp)
+{
+	const uint64 CmdListID = EventData.GetValue<uint64>("CmdListID");
+	const uint8 Pipelines = EventData.GetValue<uint8>("Pipelines");
+
+	FRHITranslateTask* TranslateJob = GetTranslateJobForCmdList(CmdListID, TEXT("activate pipelines"), ThreadID, nullptr);
+	if (!TranslateJob)
+	{
+		return;
+	}
+
+	for (int PipelineIdx = 0; PipelineIdx < 2; ++PipelineIdx)
+	{
+		const bool bActive = (Pipelines & (1 << PipelineIdx)) != 0;
+		FRHITranslateContext& TranslateContext = TranslateJob->Contexts[PipelineIdx];
+		if (bActive)
+		{
+			// Ignore if the context is already active.
+			if (TranslateContext.ActiveIntervals.IsEmpty() || TranslateContext.ActiveIntervals.Last().End != 0.0)
+			{
+				TranslateContext.ActiveIntervals.Emplace(Timestamp);
+			}
+		}
+		else
+		{
+			// Ignore if the context isn't active.
+			if (!TranslateContext.ActiveIntervals.IsEmpty() && TranslateContext.ActiveIntervals.Last().End == 0.0)
+			{
+				TranslateContext.ActiveIntervals.Last().End = Timestamp;
+			}
+		}
+	}
+}
+
+void FRenderTraceAnalyzer::ProcessSetTranslateContext(const FEventData& EventData, uint32 ThreadID, double Timestamp)
+{
+	const uint64 CmdListID = EventData.GetValue<uint64>("CmdListID");
+	const uint64 ContextID = EventData.GetValue<uint64>("ContextID");
+	const uint8 Pipeline = EventData.GetValue<uint8>("Pipeline");
+
+	uint8 PipelineIdx;
+	switch (Pipeline)
+	{
+	case 1 << 0: PipelineIdx = 0; break;
+	case 1 << 1: PipelineIdx = 1; break;
+	default:
+		UE_LOG(LogRenderTrace, Warning, TEXT("Got set translate context event for command list ID %llx with invalid pipeline mask 0x%x."), CmdListID, Pipeline);
+		return;
+	}
+
+	uint32 TranslateJobID;
+	FRHITranslateTask* TranslateJob = GetTranslateJobForCmdList(CmdListID, TEXT("set translate context"), ThreadID, &TranslateJobID);
+	if (!TranslateJob)
+	{
+		return;
+	}
+
+	FRHITranslateContext& TranslateContext = TranslateJob->Contexts[PipelineIdx];
+	if (TranslateContext.RHIContextID != 0)
+	{
+		UE_LOG(LogRenderTrace, Warning, TEXT("Got set translate context event for command list ID %llx with context ID %llx on pipe %u, but the job is already using context %llx on that pipe. Ignoring."),
+			CmdListID, ContextID, PipelineIdx, TranslateContext.RHIContextID);
+		return;
+	}
+
+	// RHIs usually have a pool of contexts, and objects are returned to the pool in their RHIFinalizeContext() implementation,
+	// which runs before the FinalizeTranslateJob event is sent. Due to threading, another translate job can pick up the context
+	// before the finalize event can remove it from LiveRHIContexts. This is why we can't check if the context is already in use;
+	// instead, we just steal it.
+	// Note that we can't send FinalizeTranslateJob before calling RHIFinalizeContext(), because that function will do more things
+	// with the context (such as signal batched sync points) and we still need to find the context -> job association when that
+	// happens. The association will still be valid when it does that, because these extra actions happen before returning the context
+	// to the pool, so another job can't steal the context while we still need the old association.
+	// Example flow with d3d12:
+	//		FRHICommandListExecutor::FTranslateState::Finalize()
+	//			|-> GDynamicRHI->RHICloseTranslateChain()
+	//			|	|-> FD3D12DynamicRHI::RHIFinalizeContext()
+	//			|		|-> FD3D12ContextCommon::Finalize()
+	//			|		|	|-> Payload->SyncPointsToWait.Append(BatchedSyncPoints.ToWait) - this needs to find the current context -> translate job mapping
+	//			|		|-> CmdContext->GetParentDevice()->ReleaseContext() - this releases the context, so another thread can pick it up for its translate job
+	// ... another thread sends a SetTranslateContext packet and steals the context for a different job ...
+	//			|-> FRenderTracing::FinalizeTranslateJob() - we should only manipulate the context if it's still bound to the job
+	TranslateContext.RHIContextID = ContextID;
+	LiveRHIContexts.Add(ContextID, TranslateJobID);
+}
+
+void FRenderTraceAnalyzer::ProcessEndTranslateCommandList(const FEventData& EventData, uint32 ThreadID, double Timestamp)
+{
+	const uint64 CmdListID = EventData.GetValue<uint64>("CmdListID");
+
+	const uint32* PermanentCmdListID = LiveCommandLists.Find(CmdListID);
+	if (!PermanentCmdListID)
+	{
+		UE_LOG(LogRenderTrace, Warning, TEXT("Got end translate event for unknown command list ID %llx."), CmdListID);
+		return;
+	}
+
+	FCommandListInstance& CmdList = Provider.EditCommandList(*PermanentCmdListID);
+	if (CmdList.TranslateTaskID == INVALID_EVENT_ID)
+	{
+		UE_LOG(LogRenderTrace, Warning, TEXT("Got end translate event for command list ID %llx which doesn't have a translate task."), CmdListID);
+		return;
+	}
+
+	FRHITranslateTask& TranslateJob = Provider.EditRHITranslateTask(CmdList.TranslateTaskID);
+	check(TranslateJob.Type == ERHITranslateTaskType::Translate);
+
+	if (TranslateJob.ThreadID != ThreadID)
+	{
+		UE_LOG(LogRenderTrace, Warning, TEXT("Got end translate event for command list ID %llx on task %llx, but the current thread %u doesn't match the job's thread %u."),
+			CmdListID, TranslateJob.AppID, ThreadID, TranslateJob.ThreadID);
+	}
+
+	if (TranslateJob.ProcessedItems.IsEmpty())
+	{
+		UE_LOG(LogRenderTrace, Warning, TEXT("Got end translate event for command list ID %llx on task %llx, but there's no open translation event."), CmdListID, TranslateJob.AppID);
+		return;
+	}
+
+	FEventInterval& CmdListEvent = TranslateJob.ProcessedItems.Last();
+	if (CmdListEvent.ItemID != *PermanentCmdListID)
+	{
+		UE_LOG(LogRenderTrace, Warning, TEXT("Got end translate event for command list ID %llx (permanent ID %x) on task %llx, but the currently open translation event is for command list permanent ID %x."),
+			CmdListID, *PermanentCmdListID, TranslateJob.AppID, CmdListEvent.ItemID);
+		return;
+	}
+
+	// End the translation event and move the end of the translate job.
+	CmdListEvent.End = Timestamp;
+	TranslateJob.EndTime = Timestamp;
+}
+
+void FRenderTraceAnalyzer::ProcessFinalizeTranslateJob(const FEventData& EventData, uint32 ThreadID, double Timestamp)
+{
+	const uint64 TranslateJobID = EventData.GetValue<uint64>("TranslateJobID");
+	const uint32* PermanentTranslateJobID = LiveTranslateJobs.Find(TranslateJobID);
+	const uint8 Flags = EventData.GetValue<uint8>("Flags");
+	if (!PermanentTranslateJobID)
+	{
+		UE_LOG(LogRenderTrace, Warning, TEXT("Got finalize event for unknown translate job ID %llx."), TranslateJobID);
+		return;
+	}
+
+	FRHITranslateTask& TranslateJob = Provider.EditRHITranslateTask(*PermanentTranslateJobID);
+	check(TranslateJob.Type == ERHITranslateTaskType::Translate);
+	TranslateJob.JobFlags = Flags;
+
+	EndTranslateTask(*PermanentTranslateJobID, Timestamp);
+
+	for (int PipeIdx = 0; PipeIdx < 2; ++PipeIdx)
+	{
+		FRHITranslateContext& TranslateContext = TranslateJob.Contexts[PipeIdx];
+		if (TranslateContext.RHIContextID == 0)
+		{
+			continue;
+		}
+
+		// The context might have been stolen by another translate job before we got the FinalizeTranslateJob
+		// notification (see the comment in ProcessSetTranslateContexts() for the full explanation). Only
+		// remove it from the live contexts map if we still own it.
+		const uint32* ExistingRHIContextMapping = LiveRHIContexts.Find(TranslateContext.RHIContextID);
+		check(ExistingRHIContextMapping);
+		if (*ExistingRHIContextMapping == *PermanentTranslateJobID)
+		{
+			LiveRHIContexts.Remove(TranslateContext.RHIContextID);
+		}
+	}
+}
+
+void FRenderTraceAnalyzer::ProcessSubmitTranslateJobs(const FEventData& EventData, uint32 ThreadID, double Timestamp)
+{
+	TConstArrayView<uint64> TranslateJobIDs = EventData.GetArrayView<uint64>("TranslateJobIDs");
+	const uint32 RTSubmissionID = EventData.GetValue<uint32>("RTSubmissionID");
+
+	// Submit "tasks" are just events, so we'll use an arbitrary duration, but we need to make sure we don't
+	// overlap the previous submit event on the submission timeline.
+	IRenderTraceProvider::TEventTimeline& Timeline = Provider.EditRHISubmissionTimeline();
+	const double EventStartTime = FMath::Max(Timestamp, Timeline.GetLastTimestamp());
+
+	TPair<uint32, FRHITranslateTask&> NewTaskEntry = Provider.AddRHITranslateTask(0, ERHITranslateTaskType::Submit, ThreadID, EventStartTime);
+	FRHITranslateTask& SubmitTask = NewTaskEntry.Value;
+
+	// Submit tasks don't have an app-provided ID, so we'll use the internal ID.
+	SubmitTask.AppID = NewTaskEntry.Key;
+	SubmitTask.EndTime = SubmitTask.StartTime + 10 * UE::Insights::FTimeValue::Microsecond;
+
+	Timeline.AppendBeginEvent(SubmitTask.StartTime, FRenderTraceEvent(NewTaskEntry.Key));
+
+	// Link the render-thread submit event to the submit task.
+	const uint32* SubmitPassID = PendingRTSubmitIDs.Find(RTSubmissionID);
+	if (SubmitPassID)
+	{
+		FRDGPassInstance& SubmitPass = Provider.EditRDGPass(*SubmitPassID);
+		check(SubmitPass.Type == ERDGPassType::Submit);
+		SubmitPass.SubmitTaskID = NewTaskEntry.Key;
+		SubmitTask.SubmitPassID = *SubmitPassID;
+		PendingRTSubmitIDs.Remove(RTSubmissionID);
+	}
+	else
+	{
+		UE_LOG(LogRenderTrace, Warning, TEXT("Got submit event with unknown render thread pass ID %x, not linking."), RTSubmissionID);
+	}
+
+	// We'll add subevents for all the submitted translate jobs. Since these don't have durations either, we'll just create
+	// equal-sized subevents for the fake duration. Note that SubEventDuration will be infinity if there are no jobs, but
+	// that's fine since the value is only used in the loop below, which will not be entered in that case.
+	const double EventDuration = SubmitTask.EndTime - SubmitTask.StartTime;
+	const double SubEventDuration = EventDuration / TranslateJobIDs.Num();
+	double SubEventStartTime = SubmitTask.StartTime;
+
+	for (uint64 TranslateJobID : TranslateJobIDs)
+	{
+		const uint32* PermanentTranslateJobID = LiveTranslateJobs.Find(TranslateJobID);
+		if (!PermanentTranslateJobID)
+		{
+			UE_LOG(LogRenderTrace, Warning, TEXT("Got submit event for unknown translate job ID %llx."), TranslateJobID);
+			continue;
+		}
+
+		FRHITranslateTask& TranslateJob = Provider.EditRHITranslateTask(*PermanentTranslateJobID);
+		check(TranslateJob.Type == ERHITranslateTaskType::Translate);
+		TranslateJob.NextPhaseTaskID = NewTaskEntry.Key;
+
+		const double SubEventEndTime = FMath::Min(SubEventStartTime + SubEventDuration, SubmitTask.EndTime);
+
+		FEventInterval& SubEvent = SubmitTask.ProcessedItems.Emplace_GetRef(*PermanentTranslateJobID, SubEventStartTime, SubEventEndTime);
+		Timeline.AppendBeginEvent(SubEvent.Start, FRenderTraceEvent(SubEvent.ItemID, NewTaskEntry.Key));
+		Timeline.AppendEndEvent(SubEvent.End);
+
+		LiveTranslateJobs.Remove(TranslateJobID);
+		SubEventStartTime = SubEventEndTime;
+	}
+
+	Timeline.AppendEndEvent(SubmitTask.EndTime);
+}
+
+void FRenderTraceAnalyzer::ProcessAssociatePlatformContext(const FEventData& EventData, uint32 ThreadID, double Timestamp)
+{
+	// Since platform contexts are pooled, we will only get this message once, when the pooled object is created. These tend to live forever,
+	// so PlatformToRHIContext is not like other mapping tables, which are constantly updated as objects are created and destroyed. Its contents
+	// persist for the entire duration of the analysis session, because we don't have explicit lifetime messages for contexts.
+	const uint64 PlatformContextID = EventData.GetValue<uint64>("PlatformContextID");
+	const uint64 RHIContextID = EventData.GetValue<uint64>("RHIContextID");
+
+	const uint64* ExistingPlatformContextMapping = PlatformToRHIContext.Find(PlatformContextID);
+	if (ExistingPlatformContextMapping)
+	{
+		if (*ExistingPlatformContextMapping != RHIContextID)
+		{
+			UE_LOG(LogRenderTrace, Warning, TEXT("Got associate platform context event from platform context ID %llx to RHI context ID %llx, but the platform context is already associated with RHI context %llx."),
+				PlatformContextID, RHIContextID, *ExistingPlatformContextMapping);
+		}
+		// Either way, we don't need to do anything.
+		return;
+	}
+
+	PlatformToRHIContext.Add(PlatformContextID, RHIContextID);
+}
+
+void FRenderTraceAnalyzer::ProcessD3D12PayloadCreated(const FEventData& EventData, uint32 ThreadID, double Timestamp)
+{
+	const uint64 PayloadID = EventData.GetValue<uint64>("PayloadID");
+	const uint8 Queue = EventData.GetValue<uint8>("Queue");
+
+	const uint32* ExistingPermanentID = LivePlatformPayloads.Find(PayloadID);
+	if (ExistingPermanentID)
+	{
+		UE_LOG(LogRenderTrace, Warning, TEXT("Got payload created event for payload ID %llx, but the payload is already live. Ignoring."), PayloadID);
+		return;
+	}
+
+	TPair<uint32, FPlatformPayload&> NewPayloadEntry = Provider.AddPlatformPayload(PayloadID, Timestamp, Queue);
+	LivePlatformPayloads.Add(PayloadID, NewPayloadEntry.Key);
+}
+
+void FRenderTraceAnalyzer::ProcessD3D12AddPayload(const FEventData& EventData, uint32 ThreadID, double Timestamp)
+{
+	const uint64 PlatformContextID = EventData.GetValue<uint64>("PlatformContextID");
+	const uint64 PayloadID = EventData.GetValue<uint64>("PayloadID");
+
+	const uint32* ExistingPermanentPayloadID = LivePlatformPayloads.Find(PayloadID);
+	if (!ExistingPermanentPayloadID)
+	{
+		UE_LOG(LogRenderTrace, Warning, TEXT("Got add payload event for unknown payload ID %llx. Ignoring."), PayloadID);
+		return;
+	}
+
+	FPlatformPayload& Payload = Provider.EditPlatformPayload(*ExistingPermanentPayloadID);
+	if (Payload.PipeIdx > 1)
+	{
+		// We ignore copy queue payloads for now.
+		return;
+	}
+
+	const uint64* ExistingPlatformContextMapping = PlatformToRHIContext.Find(PlatformContextID);
+	if (!ExistingPlatformContextMapping)
+	{
+		UE_LOG(LogRenderTrace, Warning, TEXT("Got add payload event for unknown platform context ID %llx."), PlatformContextID);
+		return;
+	}
+
+	const uint32* ExistingRHIContextMapping = LiveRHIContexts.Find(*ExistingPlatformContextMapping);
+	if (!ExistingRHIContextMapping)
+	{
+		// It's only ok to have a missing mapping if this is happening during startup when the immediate command list is being set up.
+		if (ThreadID != GameThreadID)
+		{
+			UE_LOG(LogRenderTrace, Warning, TEXT("Got add payload event for platform context ID %llx which maps to RHI context ID %llx, but the RHI context isn't live."), PlatformContextID, *ExistingPlatformContextMapping);
+		}
+		return;
+	}
+
+	const uint32 TranslateJobID = *ExistingRHIContextMapping;
+	FRHITranslateTask& TranslateJob = Provider.EditRHITranslateTask(TranslateJobID);
+
+	FRHITranslateContext& TranslateContext = TranslateJob.Contexts[Payload.PipeIdx];
+	if (TranslateContext.RHIContextID != *ExistingPlatformContextMapping)
+	{
+		UE_LOG(LogRenderTrace, Warning, TEXT("Got add payload event for platform context ID 0x%llx, payload ID 0x%llx, but the context on pipe %u is 0x%llx, not 0x%llx. Ignoring."),
+			PlatformContextID, PayloadID, Payload.PipeIdx, TranslateContext.RHIContextID, *ExistingPlatformContextMapping);
+		return;
+	}
+
+	TranslateContext.PayloadIDs.Emplace(*ExistingPermanentPayloadID);
+
+	Payload.TranslateEvent.ParentItemID = TranslateJobID;
+	Payload.TranslateEvent.ItemID = TranslateContext.PayloadIDs.Num() - 1;
+}
+
+void FRenderTraceAnalyzer::ProcessD3D12CommandListOp(const FEventData& EventData, uint32 ThreadID, double Timestamp)
+{
+	const uint64 PayloadID = EventData.GetValue<uint64>("PayloadID");
+	const uint64 D3D12CmdListID = EventData.GetValue<uint64>("D3D12CmdListID");
+	const uint8 Op = EventData.GetValue<uint8>("Op");
+
+	FPlatformPayload* Payload = FindPlatformPayload(PayloadID, TEXT("command list op"), ThreadID);
+	if (!Payload)
+	{
+		return;
+	}
+
+	switch (Op)
+	{
+	case D3D12CommandListOp_Open:
+	{
+		Payload->CmdLists.Emplace(D3D12CmdListID);
+		break;
+	}
+
+	case D3D12CommandListOp_Close:
+	{
+		// Move the end of the payload to the end of this command list, but don't exceed the translate job lifetime if it's already set.
+		Payload->EndTime = FMath::Max(Timestamp, Payload->EndTime);
+		if (Payload->TranslateEvent.ParentItemID != INVALID_EVENT_ID)
+		{
+			FRHITranslateTask& TranslateJob = Provider.EditRHITranslateTask(Payload->TranslateEvent.ParentItemID);
+			if (TranslateJob.EndTime != 0.0)
+			{
+				Payload->EndTime = FMath::Min(Payload->EndTime, TranslateJob.EndTime);
+			}
+		}
+		break;
+	}
+
+	default:
+	{
+		UE_LOG(LogRenderTrace, Warning, TEXT("Got command list open event for payload ID %llx, command list ID %llx with unknown op code %u."),
+			PayloadID, D3D12CmdListID, Op);
+		break;
+	}
+	}
+}
+
+void FRenderTraceAnalyzer::ProcessD3D12SyncPointOp(const FEventData& EventData, uint32 ThreadID, double Timestamp)
+{
+	const uint64 PayloadID = EventData.GetValue<uint64>("PayloadID");
+	const uint64 SyncPointID = EventData.GetValue<uint64>("SyncPointID");
+	const uint8 SyncPointType = EventData.GetValue<uint8>("SyncPointType");
+	const uint8 Op = EventData.GetValue<uint8>("Op");
+
+	FPlatformPayload* Payload = FindPlatformPayload(PayloadID, TEXT("sync point op"), ThreadID);
+	if (!Payload)
+	{
+		return;
+	}
+
+	const uint32* ExistingSyncPointID = LiveSyncPoints.Find(SyncPointID);
+	uint32 PermanentSyncPointID;
+	if (!ExistingSyncPointID)
+	{
+		TPair<uint32, FSyncPoint&> NewSyncPoint = Provider.AddSyncPoint(SyncPointID, static_cast<ESyncPointType>(SyncPointType));
+		PermanentSyncPointID = NewSyncPoint.Key;
+		LiveSyncPoints.Add(SyncPointID, PermanentSyncPointID);
+	}
+	else
+	{
+		PermanentSyncPointID = *ExistingSyncPointID;
+	}
+
+	auto& SyncPointList = (Op == 0) ? Payload->WaitSyncPoints : Payload->SignalSyncPoints;
+	SyncPointList.Add(PermanentSyncPointID);
+}
+
+void FRenderTraceAnalyzer::ProcessD3D12ManualFenceOp(const FEventData& EventData, uint32 ThreadID, double Timestamp)
+{
+	const uint64 PayloadID = EventData.GetValue<uint64>("PayloadID");
+	const uint64 FenceID = EventData.GetValue<uint64>("FenceID");
+	const uint64 FenceVal = EventData.GetValue<uint64>("FenceVal");
+	const uint8 Op = EventData.GetValue<uint8>("Op");
+
+	FPlatformPayload* Payload = FindPlatformPayload(PayloadID, TEXT("manual fence op"), ThreadID);
+	if (!Payload)
+	{
+		return;
+	}
+
+	const uint32* ExistingSyncPointID = ManualSyncPoints.Find(FenceID);
+	uint32 PermanentSyncPointID;
+	if (!ExistingSyncPointID)
+	{
+		TPair<uint32, FSyncPoint&> NewSyncPoint = Provider.AddSyncPoint(FenceID, ESyncPointType::Manual);
+		PermanentSyncPointID = NewSyncPoint.Key;
+		NewSyncPoint.Value.ResolvedValue = FenceVal;
+		ManualSyncPoints.Add(FenceID, PermanentSyncPointID);
+	}
+	else
+	{
+		PermanentSyncPointID = *ExistingSyncPointID;
+		FSyncPoint& ExistingSyncPoint = Provider.EditSyncPoint(PermanentSyncPointID);
+		ExistingSyncPoint.ResolvedValue = FenceVal;
+	}
+
+	auto& SyncPointList = (Op == 0) ? Payload->WaitSyncPoints : Payload->SignalSyncPoints;
+	SyncPointList.Add(PermanentSyncPointID);
+}
+
+void FRenderTraceAnalyzer::ProcessD3D12ProcessSubmissionQueueEnter(const FEventData& EventData, uint32 ThreadID, double Timestamp)
+{
+	if (PendingSubmissionBatchID != INVALID_EVENT_ID)
+	{
+		UE_LOG(LogRenderTrace, Warning, TEXT("Got new submission batch event while the previous one is still open, merging."));
+		return;
+	}
+
+	TPair<uint32, FSubmissionBatch&> NewSubmissionBatch = Provider.AddSubmissionBatch(ThreadID, Timestamp);
+	PendingSubmissionBatchID = NewSubmissionBatch.Key;
+}
+
+// This macro creates Batch and SubmissionEvent variables for convenience.
+#define ADD_SUBMIT_EVENT(Name, ...)\
+	if (PendingSubmissionBatchID == INVALID_EVENT_ID)\
+	{\
+		UE_LOG(LogRenderTrace, Warning, TEXT("Got " #Name " submission event with no open batch, ignoring."));\
+		return;\
+	}\
+	FSubmissionBatch& Batch = Provider.EditSubmissionBatch(PendingSubmissionBatchID);\
+	check(Batch.EndTime == 0.0);\
+	FSubmissionEvent& SubmissionEvent = Batch.Events.Emplace_GetRef(Timestamp, Queue);\
+	SubmissionEvent.Data.Emplace<FSubmissionEvent::F##Name>(__VA_ARGS__)
+	
+void FRenderTraceAnalyzer::ProcessD3D12SubmitYieldOnUnresolvedSyncPoint(const FEventData& EventData, uint32 ThreadID, double Timestamp)
+{
+	const uint8 Queue = EventData.GetValue<uint8>("Queue");
+	uint64 PayloadID = EventData.GetValue<uint64>("PayloadID");
+	uint64 SyncPointID = EventData.GetValue<uint64>("SyncPointID");
+
+	ADD_SUBMIT_EVENT(YieldSyncPoint, PayloadID, SyncPointID);
+}
+
+void FRenderTraceAnalyzer::ProcessD3D12SubmitYieldOnUnresolvedManualFence(const FEventData& EventData, uint32 ThreadID, double Timestamp)
+{
+	const uint8 Queue = EventData.GetValue<uint8>("Queue");
+	uint64 PayloadID = EventData.GetValue<uint64>("PayloadID");
+	uint64 FenceID = EventData.GetValue<uint64>("FenceID");
+	uint64 FenceValue = EventData.GetValue<uint64>("FenceValue");
+
+	ADD_SUBMIT_EVENT(YieldManualFence, PayloadID, FenceID, FenceValue);
+}
+
+void FRenderTraceAnalyzer::ProcessD3D12SubmitWaitQueueFence(const FEventData& EventData, uint32 ThreadID, double Timestamp)
+{
+	const uint64 PayloadID = EventData.GetValue<uint64>("PayloadID");
+	const uint8 Queue = EventData.GetValue<uint8>("ExecutingQueue");
+	const uint8 WaitOnQueue = EventData.GetValue<uint8>("WaitOnQueue");
+	const uint64 Value = EventData.GetValue<uint64>("Value");
+
+	ADD_SUBMIT_EVENT(WaitQueueFence, PayloadID, WaitOnQueue, Value);
+}
+
+void FRenderTraceAnalyzer::ProcessD3D12SubmitWaitManualFence(const FEventData& EventData, uint32 ThreadID, double Timestamp)
+{
+	const uint64 PayloadID = EventData.GetValue<uint64>("PayloadID");
+	const uint8 Queue = EventData.GetValue<uint8>("ExecutingQueue");
+	const uint64 FenceID = EventData.GetValue<uint64>("FenceID");
+	const uint64 Value = EventData.GetValue<uint64>("Value");
+
+	ADD_SUBMIT_EVENT(WaitManualFence, PayloadID, FenceID, Value);
+}
+
+void FRenderTraceAnalyzer::ProcessD3D12SubmitPlatformCommandLists(const FEventData& EventData, uint32 ThreadID, double Timestamp)
+{
+	const uint8 Queue = EventData.GetValue<uint8>("Queue");
+	TConstArrayView<uint64> PayloadIDs = EventData.GetArrayView<uint64>("PayloadIDs");
+	TConstArrayView<uint64> CmdListIDs = EventData.GetArrayView<uint64>("CmdListIDs");
+
+	TArray<uint32> PayloadPermanentIDs;
+	PayloadPermanentIDs.Reserve(PayloadIDs.Num());
+	for (uint64 PayloadID : PayloadIDs)
+	{
+		const uint32* ExistingPermanentID = LivePlatformPayloads.Find(PayloadID);
+		if (ExistingPermanentID)
+		{
+			PayloadPermanentIDs.Add(*ExistingPermanentID);
+		}
+	}
+
+	ADD_SUBMIT_EVENT(Execute, PayloadPermanentIDs, CmdListIDs);
+
+	for (uint32 PayloadID : PayloadPermanentIDs)
+	{
+		FPlatformPayload& Payload = Provider.EditPlatformPayload(PayloadID);
+		Payload.ExecutionEvent.ParentItemID = PendingSubmissionBatchID;
+		Payload.ExecutionEvent.ItemID = Batch.Events.Num() - 1;
+	}
+}
+
+void FRenderTraceAnalyzer::ProcessD3D12SubmitSignalManualFence(const FEventData& EventData, uint32 ThreadID, double Timestamp)
+{
+	const uint64 PayloadID = EventData.GetValue<uint64>("PayloadID");
+	const uint8 Queue = EventData.GetValue<uint8>("Queue");
+	const uint64 FenceID = EventData.GetValue<uint64>("FenceID");
+	const uint64 Value = EventData.GetValue<uint64>("Value");
+
+	ADD_SUBMIT_EVENT(SignalManualFence, PayloadID, FenceID, Value);
+}
+
+void FRenderTraceAnalyzer::ProcessD3D12SubmitSignalQueueFence(const FEventData& EventData, uint32 ThreadID, double Timestamp)
+{
+	const uint64 PayloadID = EventData.GetValue<uint64>("PayloadID");
+	const uint8 Queue = EventData.GetValue<uint8>("Queue");
+	const uint64 Value = EventData.GetValue<uint64>("Value");
+
+	ADD_SUBMIT_EVENT(SignalQueueFence, PayloadID, Value);
+}
+
+void FRenderTraceAnalyzer::ProcessD3D12SubmitResolveSyncPoint(const FEventData& EventData, uint32 ThreadID, double Timestamp)
+{
+	const uint8 Queue = EventData.GetValue<uint8>("Queue");
+	const uint64 PayloadID = EventData.GetValue<uint64>("PayloadID");
+	const uint64 SyncPointID = EventData.GetValue<uint64>("SyncPointID");
+	const uint64 Value = EventData.GetValue<uint64>("Value");
+
+	const uint32* PermanentSyncPointID = LiveSyncPoints.Find(SyncPointID);
+	if (!PermanentSyncPointID)
+	{
+		UE_LOG(LogRenderTrace, Warning, TEXT("Got resolve sync point event for unknown sync point ID %llx."), SyncPointID);
+		return;
+	}
+
+	const uint32* PermanentPayloadID = LivePlatformPayloads.Find(PayloadID);
+	if (!PermanentPayloadID)
+	{
+		UE_LOG(LogRenderTrace, Warning, TEXT("Got resolve sync point event for sync point ID %llx with unknown payload ID %llx."), SyncPointID, PayloadID);
+		return;
+	}
+
+	ADD_SUBMIT_EVENT(ResolveSyncPoint, *PermanentPayloadID, *PermanentSyncPointID, Value);
+
+	FSyncPoint& SyncPoint = Provider.EditSyncPoint(*PermanentSyncPointID);
+	check(SyncPoint.ResolvedValue == 0);
+	SyncPoint.ResolvedValue = Value;
+	SyncPoint.ResolvedByPayload = *PermanentPayloadID;
+
+	FPlatformPayload& Payload = Provider.EditPlatformPayload(*PermanentPayloadID);
+	Payload.ResolveSyncPointEvents.Emplace(Batch.Events.Num() - 1, PendingSubmissionBatchID);
+
+	LiveSyncPoints.Remove(SyncPointID);
+}
+
+void FRenderTraceAnalyzer::EndPendingSubmissionBatch(double Timestamp, uint8 ExitStatus)
+{
+	check(PendingSubmissionBatchID != INVALID_EVENT_ID);
+
+	FSubmissionBatch& Batch = Provider.EditSubmissionBatch(PendingSubmissionBatchID);
+	check(Batch.EndTime == 0.0);
+	Batch.EndTime = Timestamp;
+	Batch.ExitStatus = ExitStatus;
+
+	IRenderTraceProvider::TEventTimeline& Timeline = Provider.EditSubmissionQueueTimeline();
+	Timeline.AppendBeginEvent(Batch.StartTime, FRenderTraceEvent(PendingSubmissionBatchID));
+	for (int EventIdx = 0; EventIdx < Batch.Events.Num(); ++EventIdx)
+	{
+		const FSubmissionEvent& Event = Batch.Events[EventIdx];
+		Timeline.AppendBeginEvent(Event.Timestamp, FRenderTraceEvent(EventIdx, PendingSubmissionBatchID));
+		const double EventEndTime = EventIdx < Batch.Events.Num() - 1 ? Batch.Events[EventIdx + 1].Timestamp : Batch.EndTime;
+		Timeline.AppendEndEvent(EventEndTime);
+	}
+	Timeline.AppendEndEvent(Batch.EndTime);
+
+	PendingSubmissionBatchID = INVALID_EVENT_ID;
+}
+
+void FRenderTraceAnalyzer::ProcessD3D12ProcessSubmissionQueueExit(const FEventData& EventData, uint32 ThreadID, double Timestamp)
+{
+	if (PendingSubmissionBatchID == INVALID_EVENT_ID)
+	{
+		UE_LOG(LogRenderTrace, Warning, TEXT("Got end submission batch event with no open batch, ignoring."));
+		return;
+	}
+
+	const uint8 ExitStatus = EventData.GetValue<uint8>("Status");
+	EndPendingSubmissionBatch(Timestamp, ExitStatus);
+}
+
+#undef ADD_SUBMIT_EVENT
+
+void FRenderTraceAnalyzer::ProcessD3D12ProcessInterruptQueueEnter(const FEventData& EventData, uint32 ThreadID, double Timestamp)
+{
+	if (PendingInterruptWakeUpID != INVALID_EVENT_ID)
+	{
+		UE_LOG(LogRenderTrace, Warning, TEXT("Got new interrupt wake up event while the previous one is still open, merging."));
+		return;
+	}
+
+	TPair<uint32, FInterruptWakeUp&> NewWakeUp = Provider.AddInterruptWakeUp(ThreadID, Timestamp);
+	PendingInterruptWakeUpID = NewWakeUp.Key;
+}
+
+void FRenderTraceAnalyzer::ProcessD3D12InterruptQueueFenceSignaled(const FEventData& EventData, uint32 ThreadID, double Timestamp)
+{
+	const uint8 Queue = EventData.GetValue<uint8>("Queue");
+	const uint64 PendingPayloadID = EventData.GetValue<uint64>("PendingPayloadID");
+	const uint64 CurrentFenceValue = EventData.GetValue<uint64>("CurrentFenceValue");
+	const uint64 LastCPUSignaledFenceValue = EventData.GetValue<uint64>("LastCPUSignaledFenceValue");
+
+	if (PendingInterruptWakeUpID == INVALID_EVENT_ID)
+	{
+		UE_LOG(LogRenderTrace, Warning, TEXT("Got fence signaled event with no open interrupt wake up, ignoring."));
+		return;
+	}
+
+	const uint32* ExistingPermanentID = LivePlatformPayloads.Find(PendingPayloadID);
+	if (!ExistingPermanentID)
+	{
+		UE_LOG(LogRenderTrace, Warning, TEXT("Got fence signaled event for unknown payload 0x%llx, ignoring."), PendingPayloadID);
+		return;
+	}
+
+	FInterruptWakeUp& WakeUp = Provider.EditInterruptWakeUp(PendingInterruptWakeUpID);
+	WakeUp.SignalEvents.Emplace(Timestamp, Queue, *ExistingPermanentID, CurrentFenceValue, LastCPUSignaledFenceValue);
+
+	FPlatformPayload& Payload = Provider.EditPlatformPayload(*ExistingPermanentID);
+	Payload.InterruptEvent.ParentItemID = PendingInterruptWakeUpID;
+	Payload.InterruptEvent.ItemID = WakeUp.SignalEvents.Num() - 1;
+
+	LivePlatformPayloads.Remove(PendingPayloadID);
+}
+
+void FRenderTraceAnalyzer::EndPendingInterruptWakeUp(double Timestamp, uint8 ExitStatus)
+{
+	check(PendingInterruptWakeUpID != INVALID_EVENT_ID);
+
+	FInterruptWakeUp& WakeUp = Provider.EditInterruptWakeUp(PendingInterruptWakeUpID);
+	check(WakeUp.EndTime == 0.0);
+	WakeUp.EndTime = Timestamp;
+	WakeUp.ExitStatus = ExitStatus;
+
+	IRenderTraceProvider::TEventTimeline& Timeline = Provider.EditInterruptTimeline();
+	Timeline.AppendBeginEvent(WakeUp.StartTime, FRenderTraceEvent(PendingInterruptWakeUpID));
+	for (int EventIdx = 0; EventIdx < WakeUp.SignalEvents.Num(); ++EventIdx)
+	{
+		const FInterruptFenceSignaledEvent& Event = WakeUp.SignalEvents[EventIdx];
+		Timeline.AppendBeginEvent(Event.Timestamp, FRenderTraceEvent(EventIdx, PendingInterruptWakeUpID));
+		const double EventEndTime = EventIdx < WakeUp.SignalEvents.Num() - 1 ? WakeUp.SignalEvents[EventIdx + 1].Timestamp : WakeUp.EndTime;
+		Timeline.AppendEndEvent(EventEndTime);
+	}
+	Timeline.AppendEndEvent(WakeUp.EndTime);
+
+	PendingInterruptWakeUpID = INVALID_EVENT_ID;
+}
+
+void FRenderTraceAnalyzer::ProcessD3D12ProcessInterruptQueueExit(const FEventData& EventData, uint32 ThreadID, double Timestamp)
+{
+	if (PendingInterruptWakeUpID == INVALID_EVENT_ID)
+	{
+		UE_LOG(LogRenderTrace, Warning, TEXT("Got end interrupt wake up event with no open event, ignoring."));
+		return;
+	}
+
+	const uint8 ExitStatus = EventData.GetValue<uint8>("Status");
+	EndPendingInterruptWakeUp(Timestamp, ExitStatus);
+}
+
+bool FRenderTraceAnalyzer::CacheThreadIDs()
+{
+	// If everything is cached there's nothing we need to do.
+	if (GameThreadID != 0 && RenderThreadID != 0 && RHIThreadID != 0)
+	{
+		return false;
+	}
+
+	TraceServices::FAnalysisSessionReadScope SessionReadScope(Session);
+	const TraceServices::IThreadProvider& ThreadProvider = TraceServices::ReadThreadProvider(Session);
+	if (ThreadProvider.GetModCount() == LastThreadProviderModCount)
+	{
+		// The thread provider state hasn't changed since last time we checked, there's not point checking again.
+		return false;
+	}
+
+	bool bThreadAdded = false;
+	ThreadProvider.EnumerateThreads([this, &bThreadAdded](const TraceServices::FThreadInfo& ThreadInfo)
+	{
+		if (FCString::Strcmp(ThreadInfo.Name, TEXT("GameThread")) == 0)
+		{
+			GameThreadID = ThreadInfo.Id;
+		}
+		else if (FCString::Strncmp(ThreadInfo.Name, TEXT("RenderThread"), 12) == 0)
+		{
+			if (RenderThreadID == 0)
+			{
+				RenderThreadID = ThreadInfo.Id;
+				bThreadAdded = true;
+			}
+		}
+		else if (FCString::Strncmp(ThreadInfo.Name, TEXT("RHIThread"), 9) == 0)
+		{
+			if (RHIThreadID == 0)
+			{
+				RHIThreadID = ThreadInfo.Id;
+				bThreadAdded = true;
+			}
+		}
+	});
+
+	LastThreadProviderModCount = ThreadProvider.GetModCount();
+	return bThreadAdded;
+}
+
+bool FRenderTraceAnalyzer::OnEvent(uint16 RouteId, EStyle Style, const FOnEventContext& EventContext)
+{
+	LLM_SCOPE_BYNAME(TEXT("Insights/FRenderTraceAnalyzer"));
+
+	const FEventData& EventData = EventContext.EventData;
+	const uint32 ThreadID = EventContext.ThreadInfo.GetId();
+	const double Timestamp = EventContext.EventTime.AsSeconds(EventData.GetValue<uint64>("Timestamp"));
+
+	bool bThreadAdded = CacheThreadIDs();
+
+	FinalTimestamp = FMath::Max(Timestamp, FinalTimestamp);
+
+	double CurrentDuration;
+	{
+		TraceServices::FAnalysisSessionReadScope SessionReadScope(Session);
+		CurrentDuration = Session.GetDurationSeconds();
+	}
+
+	if (FinalTimestamp - CurrentDuration > 5.0)
+	{
+		TraceServices::FAnalysisSessionEditScope SessionEditScope(Session);
+		Session.UpdateDurationSeconds(FinalTimestamp);
+	}
+
+	TraceServices::FProviderEditScopeLock ProviderEditScope(Provider);
+
+	if (bThreadAdded)
+	{
+		// Make sure the render thread and RHI thread timelines come first in their respective categories. It's fine
+		// to call these functions multiple times when the RT and RHIT appear at different points in time since they
+		// will only add a timeline the first time they're called.
+		if (RenderThreadID != 0)
+		{
+			Provider.EditRDGThreadTimeline(RenderThreadID);
+		}
+
+		if (RHIThreadID != 0)
+		{
+			Provider.EditRHITranslateThreadTimeline(RHIThreadID);
+		}
+	}
+
+#define DISPATCH_EVENT(Name) case Packet_##Name: Process##Name(EventData, ThreadID, Timestamp); break
+
+	switch (RouteId)
+	{
+	DISPATCH_EVENT(CommandListCreated);
+	DISPATCH_EVENT(CommandListDetach);
+	DISPATCH_EVENT(CommandListFinishRecording);
+	DISPATCH_EVENT(CommandListDestroyed);
+	DISPATCH_EVENT(AddDispatchPassCommandList);
+	DISPATCH_EVENT(SubmitCommandLists);
+	DISPATCH_EVENT(SplitTranslateJob);
+	DISPATCH_EVENT(DispatchCommandList);
+	DISPATCH_EVENT(BeginTranslateCommandList);
+	DISPATCH_EVENT(ActivatePipelines);
+	DISPATCH_EVENT(SetTranslateContext);
+	DISPATCH_EVENT(EndTranslateCommandList);
+	DISPATCH_EVENT(FinalizeTranslateJob);
+	DISPATCH_EVENT(SubmitTranslateJobs);
+	DISPATCH_EVENT(AssociatePlatformContext);
+	DISPATCH_EVENT(D3D12PayloadCreated);
+	DISPATCH_EVENT(D3D12AddPayload);
+	DISPATCH_EVENT(D3D12CommandListOp);
+	DISPATCH_EVENT(D3D12SyncPointOp);
+	DISPATCH_EVENT(D3D12ManualFenceOp);
+	DISPATCH_EVENT(D3D12ProcessSubmissionQueueEnter);
+	DISPATCH_EVENT(D3D12SubmitYieldOnUnresolvedSyncPoint);
+	DISPATCH_EVENT(D3D12SubmitYieldOnUnresolvedManualFence);
+	DISPATCH_EVENT(D3D12SubmitWaitQueueFence);
+	DISPATCH_EVENT(D3D12SubmitWaitManualFence);
+	DISPATCH_EVENT(D3D12SubmitPlatformCommandLists);
+	DISPATCH_EVENT(D3D12SubmitSignalManualFence);
+	DISPATCH_EVENT(D3D12SubmitSignalQueueFence);
+	DISPATCH_EVENT(D3D12SubmitResolveSyncPoint);
+	DISPATCH_EVENT(D3D12ProcessSubmissionQueueExit);
+	DISPATCH_EVENT(D3D12ProcessInterruptQueueEnter);
+	DISPATCH_EVENT(D3D12InterruptQueueFenceSignaled);
+	DISPATCH_EVENT(D3D12ProcessInterruptQueueExit);
+
+	case Packet_BeginExecuteRDGPass: ProcessBeginExecuteRDGPass(EventData, ThreadID, Timestamp, false); break;
+	case Packet_BeginExecuteDispatchPass: ProcessBeginExecuteRDGPass(EventData, ThreadID, Timestamp, true); break;
+	case Packet_EndExecuteRDGPass: ProcessEndExecuteRDGPass(EventData, ThreadID, Timestamp, false); break;
+	case Packet_EndExecuteDispatchPass: ProcessEndExecuteRDGPass(EventData, ThreadID, Timestamp, true); break;
+	}
+
+#undef DISPATCH_EVENT
+
+	return true;
+}
+
+void FRenderTraceAnalyzer::OnAnalysisEnd()
+{
+	{
+		TraceServices::FAnalysisSessionEditScope SessionEditScope(Session);
+		Session.UpdateDurationSeconds(FinalTimestamp);
+	}
+
+	TraceServices::FProviderEditScopeLock ProviderEditScope(Provider);
+
+	// Close pending events with an infinity end timestamp so that Insights knows that they never ended.
+	const double Timestamp = std::numeric_limits<double>::infinity();
+
+	// The trace might be truncated, so we need to close any pending objects.
+	if (!LiveRDGPasses.IsEmpty())
+	{
+		UE_LOG(LogRenderTrace, Warning, TEXT("Analysis ended with %u pending RDG passes, force-closing them."), LiveRDGPasses.Num());
+		for (const auto& Pair : LiveRDGPasses)
+		{
+			EndRDGPassExecute(Pair.Value, Timestamp);
+		}
+		LiveRDGPasses.Empty();
+	}
+
+	if (!LiveCommandLists.IsEmpty())
+	{
+		UE_LOG(LogRenderTrace, Warning, TEXT("Analysis ended with %u pending RHI command lists, force-closing them."), LiveCommandLists.Num());
+		for (const auto& Pair : LiveCommandLists)
+		{
+			EndCommandList(Pair.Value, Timestamp);
+		}
+		LiveCommandLists.Empty();
+	}
+
+	if (!LiveTranslateJobs.IsEmpty())
+	{
+		UE_LOG(LogRenderTrace, Warning, TEXT("Analysis ended with %u pending translate jobs, force-closing them."), LiveTranslateJobs.Num());
+		for (const auto& Pair : LiveTranslateJobs)
+		{
+			EndTranslateTask(Pair.Value, Timestamp);
+		}
+		LiveTranslateJobs.Empty();
+	}
+
+	if (!LiveRHIContexts.IsEmpty())
+	{
+		UE_LOG(LogRenderTrace, Warning, TEXT("Analysis ended with %u pending RHI contexts."), LiveRHIContexts.Num());
+		LiveRHIContexts.Empty();
+	}
+
+	PlatformToRHIContext.Empty();
+
+	if (!PendingRTSubmitIDs.IsEmpty())
+	{
+		UE_LOG(LogRenderTrace, Warning, TEXT("Analysis ended with %u pending render thread submits."), PendingRTSubmitIDs.Num());
+		PendingRTSubmitIDs.Empty();
+	}
+
+	if (!LivePlatformPayloads.IsEmpty())
+	{
+		UE_LOG(LogRenderTrace, Warning, TEXT("Analysis ended with %u pending platform payloads."), LivePlatformPayloads.Num());
+		LivePlatformPayloads.Empty();
+	}
+
+	if (!LiveSyncPoints.IsEmpty())
+	{
+		UE_LOG(LogRenderTrace, Warning, TEXT("Analysis ended with %u pending sync points."), LiveSyncPoints.Num());
+		LiveSyncPoints.Empty();
+	}
+
+	ManualSyncPoints.Empty();
+
+	if (PendingSubmissionBatchID != INVALID_EVENT_ID)
+	{
+		UE_LOG(LogRenderTrace, Warning, TEXT("Analysis ended a pending submission batch, force-closing it."));
+		EndPendingSubmissionBatch(Timestamp, 0);
+	}
+
+	if (PendingInterruptWakeUpID != INVALID_EVENT_ID)
+	{
+		UE_LOG(LogRenderTrace, Warning, TEXT("Analysis ended a pending interrupt wake up, force-closing it."));
+		EndPendingInterruptWakeUp(Timestamp, 0);
+	}
+}
+
+} //namespace RenderTraceInsights
+} //namespace UE
diff --git a/Engine/Plugins/RenderTraceInsights/Source/RenderTraceInsights/Private/RenderTraceAnalyzer.h b/Engine/Plugins/RenderTraceInsights/Source/RenderTraceInsights/Private/RenderTraceAnalyzer.h
new file mode 100644
--- /dev/null
+++ b/Engine/Plugins/RenderTraceInsights/Source/RenderTraceInsights/Private/RenderTraceAnalyzer.h
@@ -0,0 +1,139 @@
+// Copyright Epic Games, Inc. All Rights Reserved.
+
+#pragma once
+
+#include "Trace/Analyzer.h"
+#include "IRenderTraceProvider.h"
+
+namespace TraceServices { class IAnalysisSession; }
+
+namespace UE
+{
+namespace RenderTraceInsights
+{ 
+
+class FRenderTraceAnalyzer : public UE::Trace::IAnalyzer
+{
+public:
+	FRenderTraceAnalyzer(TraceServices::IAnalysisSession& InSession, IEditableRenderTraceProvider& InProvider);
+
+	virtual void OnAnalysisBegin(const FOnAnalysisContext& Context) override;
+	virtual void OnAnalysisEnd() override;
+	virtual bool OnEvent(uint16 RouteId, EStyle Style, const FOnEventContext& EventContext) override;
+
+private:
+	enum : uint16
+	{
+		Packet_CommandListCreated,
+		Packet_CommandListDetach,
+		Packet_CommandListFinishRecording,
+		Packet_CommandListDestroyed,
+		Packet_BeginExecuteRDGPass,
+		Packet_EndExecuteRDGPass,
+		Packet_BeginExecuteDispatchPass,
+		Packet_AddDispatchPassCommandList,
+		Packet_EndExecuteDispatchPass,
+		Packet_SubmitCommandLists,
+		Packet_SplitTranslateJob,
+		Packet_DispatchCommandList,
+		Packet_BeginTranslateCommandList,
+		Packet_ActivatePipelines,
+		Packet_SetTranslateContext,
+		Packet_EndTranslateCommandList,
+		Packet_FinalizeTranslateJob,
+		Packet_SubmitTranslateJobs,
+		Packet_AssociatePlatformContext,
+		Packet_D3D12PayloadCreated,
+		Packet_D3D12AddPayload,
+		Packet_D3D12CommandListOp,
+		Packet_D3D12SyncPointOp,
+		Packet_D3D12ManualFenceOp,
+		Packet_D3D12ProcessSubmissionQueueEnter,
+		Packet_D3D12SubmitYieldOnUnresolvedSyncPoint,
+		Packet_D3D12SubmitYieldOnUnresolvedManualFence,
+		Packet_D3D12SubmitWaitQueueFence,
+		Packet_D3D12SubmitWaitManualFence,
+		Packet_D3D12SubmitPlatformCommandLists,
+		Packet_D3D12SubmitSignalManualFence,
+		Packet_D3D12SubmitSignalQueueFence,
+		Packet_D3D12SubmitResolveSyncPoint,
+		Packet_D3D12ProcessSubmissionQueueExit,
+		Packet_D3D12ProcessInterruptQueueEnter,
+		Packet_D3D12InterruptQueueFenceSignaled,
+		Packet_D3D12ProcessInterruptQueueExit,
+	};
+
+	TraceServices::IAnalysisSession& Session;
+	IEditableRenderTraceProvider& Provider;
+
+	uint32 GameThreadID = 0;
+	uint32 RenderThreadID = 0;
+	uint32 RHIThreadID = 0;
+	TMap<uint64, uint32> LiveRDGPasses;
+	TMap<uint64, uint32> LiveCommandLists;
+	TMap<uint64, uint32> LiveTranslateJobs;
+	TMap<uint64, uint32> LiveRHIContexts;
+	TMap<uint64, uint64> PlatformToRHIContext;
+	TMap<uint32, uint32> PendingRTSubmitIDs;
+	TMap<uint64, uint32> LivePlatformPayloads;
+	TMap<uint64, uint32> LiveSyncPoints;
+	TMap<uint64, uint32> ManualSyncPoints;
+	uint32 PendingSubmissionBatchID = INVALID_EVENT_ID;
+	uint32 PendingInterruptWakeUpID = INVALID_EVENT_ID;
+	double FinalTimestamp = 0.0;
+	uint64 LastThreadProviderModCount = 0;
+
+	void BeginCmdListExecute(uint64 CmdListAppID, uint32 PermanentPassID, double Timestamp);
+	void EndRDGPassExecute(uint32 PermanentPassID, double Timestamp);
+	void AddCommandListToTimeline(uint32 PermanentCmdListID);
+	void EndCommandList(uint32 PermanentCmdListID, double Timestamp);
+	void EndTranslateTask(uint32 PermanentTaskID, double Timestamp);
+	void EndPendingSubmissionBatch(double Timestamp, uint8 ExitStatus);
+	void EndPendingInterruptWakeUp(double Timestamp, uint8 ExitStatus);
+	FRHITranslateTask* GetTranslateJobForCmdList(uint64 CmdListID, const TCHAR* EventName, uint32 ThreadID, uint32* OutTranslateJobID);
+	FPlatformPayload* FindPlatformPayload(uint64 PayloadID, const TCHAR* EventName, uint32 ThreadID);
+	bool CacheThreadIDs();
+
+#define DECLARE_EVENT_PROCESSING_FUNC(EventName) void Process##EventName(const FEventData& EventData, uint32 ThreadID, double Timestamp)
+	DECLARE_EVENT_PROCESSING_FUNC(CommandListCreated);
+	DECLARE_EVENT_PROCESSING_FUNC(CommandListDetach);
+	DECLARE_EVENT_PROCESSING_FUNC(CommandListFinishRecording);
+	DECLARE_EVENT_PROCESSING_FUNC(CommandListDestroyed);
+	DECLARE_EVENT_PROCESSING_FUNC(AddDispatchPassCommandList);
+	DECLARE_EVENT_PROCESSING_FUNC(SubmitCommandLists);
+	DECLARE_EVENT_PROCESSING_FUNC(SplitTranslateJob);
+	DECLARE_EVENT_PROCESSING_FUNC(DispatchCommandList);
+	DECLARE_EVENT_PROCESSING_FUNC(BeginTranslateCommandList);
+	DECLARE_EVENT_PROCESSING_FUNC(ActivatePipelines);
+	DECLARE_EVENT_PROCESSING_FUNC(SetTranslateContext);
+	DECLARE_EVENT_PROCESSING_FUNC(EndTranslateCommandList);
+	DECLARE_EVENT_PROCESSING_FUNC(FinalizeTranslateJob);
+	DECLARE_EVENT_PROCESSING_FUNC(SubmitTranslateJobs);
+	DECLARE_EVENT_PROCESSING_FUNC(AssociatePlatformContext);
+	DECLARE_EVENT_PROCESSING_FUNC(D3D12PayloadCreated);
+	DECLARE_EVENT_PROCESSING_FUNC(D3D12AddPayload);
+	DECLARE_EVENT_PROCESSING_FUNC(D3D12CommandListOp);
+	DECLARE_EVENT_PROCESSING_FUNC(D3D12SyncPointOp);
+	DECLARE_EVENT_PROCESSING_FUNC(D3D12ManualFenceOp);
+	DECLARE_EVENT_PROCESSING_FUNC(D3D12ProcessSubmissionQueueEnter);
+	DECLARE_EVENT_PROCESSING_FUNC(D3D12SubmitYieldOnUnresolvedSyncPoint);
+	DECLARE_EVENT_PROCESSING_FUNC(D3D12SubmitYieldOnUnresolvedManualFence);
+	DECLARE_EVENT_PROCESSING_FUNC(D3D12SubmitWaitQueueFence);
+	DECLARE_EVENT_PROCESSING_FUNC(D3D12SubmitWaitManualFence);
+	DECLARE_EVENT_PROCESSING_FUNC(D3D12SubmitPlatformCommandLists);
+	DECLARE_EVENT_PROCESSING_FUNC(D3D12SubmitSignalManualFence);
+	DECLARE_EVENT_PROCESSING_FUNC(D3D12SubmitSignalQueueFence);
+	DECLARE_EVENT_PROCESSING_FUNC(D3D12SubmitResolveSyncPoint);
+	DECLARE_EVENT_PROCESSING_FUNC(D3D12ProcessSubmissionQueueExit);
+	DECLARE_EVENT_PROCESSING_FUNC(D3D12ProcessInterruptQueueEnter);
+	DECLARE_EVENT_PROCESSING_FUNC(D3D12InterruptQueueFenceSignaled);
+	DECLARE_EVENT_PROCESSING_FUNC(D3D12ProcessInterruptQueueExit);
+#undef DECLARE_EVENT_PROCESSING_FUNC
+
+	// These are shared between regular and dispatch RDG passes so they have different signatures from the other event processing methods.
+	void ProcessBeginExecuteRDGPass(const FEventData& EventData, uint32 ThreadID, double Timestamp, bool bIsDispatch);
+	void ProcessEndExecuteRDGPass(const FEventData& EventData, uint32 ThreadID, double Timestamp, bool bIsDispatch);
+};
+
+} //namespace RenderTraceInsights
+} //namespace UE
diff --git a/Engine/Plugins/RenderTraceInsights/Source/RenderTraceInsights/Private/RenderTraceInsightsModule.cpp b/Engine/Plugins/RenderTraceInsights/Source/RenderTraceInsights/Private/RenderTraceInsightsModule.cpp
new file mode 100644
--- /dev/null
+++ b/Engine/Plugins/RenderTraceInsights/Source/RenderTraceInsights/Private/RenderTraceInsightsModule.cpp
@@ -0,0 +1,33 @@
+// Copyright Epic Games, Inc. All Rights Reserved.
+
+#include "RenderTraceInsightsModule.h"
+#include "Modules/ModuleManager.h"
+#include "Features/IModularFeatures.h"
+
+
+namespace UE
+{
+namespace RenderTraceInsights
+{
+
+FRenderTraceInsightsModule& FRenderTraceInsightsModule::Get()
+{
+	return FModuleManager::LoadModuleChecked<FRenderTraceInsightsModule>("RenderTraceInsights");
+}
+
+void FRenderTraceInsightsModule::StartupModule()
+{
+	IModularFeatures::Get().RegisterModularFeature(TraceServices::ModuleFeatureName, &TraceModule);
+	IModularFeatures::Get().RegisterModularFeature(UE::Insights::Timing::TimingViewExtenderFeatureName, &TimingViewExtender);
+}
+
+void FRenderTraceInsightsModule::ShutdownModule()
+{
+	IModularFeatures::Get().UnregisterModularFeature(TraceServices::ModuleFeatureName, &TraceModule);
+	IModularFeatures::Get().UnregisterModularFeature(UE::Insights::Timing::TimingViewExtenderFeatureName, &TimingViewExtender);
+}
+
+} //namespace RenderTraceInsights
+} //namespace UE
+
+IMPLEMENT_MODULE(UE::RenderTraceInsights::FRenderTraceInsightsModule, RenderTraceInsights);
diff --git a/Engine/Plugins/RenderTraceInsights/Source/RenderTraceInsights/Private/RenderTraceInsightsModule.h b/Engine/Plugins/RenderTraceInsights/Source/RenderTraceInsights/Private/RenderTraceInsightsModule.h
new file mode 100644
--- /dev/null
+++ b/Engine/Plugins/RenderTraceInsights/Source/RenderTraceInsights/Private/RenderTraceInsightsModule.h
@@ -0,0 +1,31 @@
+// Copyright Epic Games, Inc. All Rights Reserved.
+
+#pragma once
+
+#include "Modules/ModuleInterface.h"
+
+#include "RenderTraceModule.h"
+#include "RenderTraceTimingViewExtender.h"
+
+namespace UE
+{
+namespace RenderTraceInsights
+{
+
+class FRenderTraceInsightsModule : public IModuleInterface
+{
+public:
+	static FRenderTraceInsightsModule& Get();
+
+	//~ Begin IModuleInterface interface
+	virtual void StartupModule() override;
+	virtual void ShutdownModule() override;
+	//~ End IModuleInterface interface
+
+private:
+	FRenderTraceModule TraceModule;
+	FRenderTraceTimingViewExtender TimingViewExtender;
+};
+
+} //namespace RenderTraceInsights
+} //namespace UE
diff --git a/Engine/Plugins/RenderTraceInsights/Source/RenderTraceInsights/Private/RenderTraceModule.cpp b/Engine/Plugins/RenderTraceInsights/Source/RenderTraceInsights/Private/RenderTraceModule.cpp
new file mode 100644
--- /dev/null
+++ b/Engine/Plugins/RenderTraceInsights/Source/RenderTraceInsights/Private/RenderTraceModule.cpp
@@ -0,0 +1,36 @@
+// Copyright Epic Games, Inc. All Rights Reserved.
+
+#include "RenderTraceModule.h"
+#include "RenderTraceProvider.h"
+#include "RenderTraceAnalyzer.h"
+
+DEFINE_LOG_CATEGORY(LogRenderTrace);
+
+namespace UE
+{
+namespace RenderTraceInsights
+{
+
+FName FRenderTraceModule::ModuleName("TraceModule_RenderTrace");
+
+void FRenderTraceModule::GetModuleInfo(TraceServices::FModuleInfo& OutModuleInfo)
+{
+	OutModuleInfo.Name = ModuleName;
+	OutModuleInfo.DisplayName = TEXT("RenderTrace");
+}
+
+void FRenderTraceModule::OnAnalysisBegin(TraceServices::IAnalysisSession& InSession)
+{
+	TSharedPtr<FRenderTraceProvider> Provider = MakeShared<FRenderTraceProvider>(InSession);
+	InSession.AddProvider(FRenderTraceProvider::ProviderName, Provider);
+
+	InSession.AddAnalyzer(new FRenderTraceAnalyzer(InSession, *Provider));
+}
+
+void FRenderTraceModule::GetLoggers(TArray<const TCHAR *>& OutLoggers)
+{
+	OutLoggers.Add(TEXT("RenderTrace"));
+}
+
+} //namespace RenderTraceInsights
+} //namespace UE
diff --git a/Engine/Plugins/RenderTraceInsights/Source/RenderTraceInsights/Private/RenderTraceModule.h b/Engine/Plugins/RenderTraceInsights/Source/RenderTraceInsights/Private/RenderTraceModule.h
new file mode 100644
--- /dev/null
+++ b/Engine/Plugins/RenderTraceInsights/Source/RenderTraceInsights/Private/RenderTraceModule.h
@@ -0,0 +1,29 @@
+// Copyright Epic Games, Inc. All Rights Reserved.
+
+#pragma once
+
+#include "TraceServices/ModuleService.h"
+
+DECLARE_LOG_CATEGORY_EXTERN(LogRenderTrace, Log, All);
+
+namespace UE
+{
+namespace RenderTraceInsights
+{
+
+class FRenderTraceModule : public TraceServices::IModule
+{
+public:
+	//~ Begin TraceServices::IModule interface
+	virtual void GetModuleInfo(TraceServices::FModuleInfo& OutModuleInfo) override;
+	virtual void OnAnalysisBegin(TraceServices::IAnalysisSession& Session) override;
+	virtual void GetLoggers(TArray<const TCHAR *>& OutLoggers) override;
+	virtual const TCHAR* GetCommandLineArgument() override { return TEXT("rendertrace"); }
+	//~ End TraceServices::IModule interface
+
+private:
+	static FName ModuleName;
+};
+
+} //namespace RenderTraceInsights
+} //namespace UE
diff --git a/Engine/Plugins/RenderTraceInsights/Source/RenderTraceInsights/Private/RenderTraceProvider.cpp b/Engine/Plugins/RenderTraceInsights/Source/RenderTraceInsights/Private/RenderTraceProvider.cpp
new file mode 100644
--- /dev/null
+++ b/Engine/Plugins/RenderTraceInsights/Source/RenderTraceInsights/Private/RenderTraceProvider.cpp
@@ -0,0 +1,408 @@
+// Copyright Epic Games, Inc. All Rights Reserved.
+
+#include "RenderTraceProvider.h"
+#include "RenderTraceModule.h"
+
+namespace UE
+{
+namespace RenderTraceInsights
+{
+
+thread_local TraceServices::FProviderLock::FThreadLocalState GRenderTraceProviderLockState;
+FName FRenderTraceProvider::ProviderName("RenderTraceProvider");
+
+FName GetRenderTraceProviderName()
+{
+	return FRenderTraceProvider::ProviderName;
+}
+
+const IRenderTraceProvider* ReadRenderTraceProvider(const TraceServices::IAnalysisSession& Session)
+{
+	return Session.ReadProvider<IRenderTraceProvider>(FRenderTraceProvider::ProviderName);
+}
+
+IEditableRenderTraceProvider* EditRenderTraceProvider(TraceServices::IAnalysisSession& Session)
+{
+	return Session.EditProvider<IEditableRenderTraceProvider>(FRenderTraceProvider::ProviderName);
+}
+
+FRenderTraceProvider::FRenderTraceProvider(TraceServices::IAnalysisSession& InSession)
+	: Session(InSession)
+{
+	RTSubmissionTimeline = MakeUnique<TEventTimeline>(Session.GetLinearAllocator());
+	RHISubmissionTimeline = MakeUnique<TEventTimeline>(Session.GetLinearAllocator());
+	SubmissionQueueTimeline = MakeUnique<TEventTimeline>(Session.GetLinearAllocator());
+	InterruptTimeline = MakeUnique<TEventTimeline>(Session.GetLinearAllocator());
+}
+
+int32 FRenderTraceProvider::GetNumCommandLists() const
+{
+	ReadAccessCheck();
+	return CommandLists.Num();
+}
+
+const FCommandListInstance& FRenderTraceProvider::GetCommandList(uint32 ID) const
+{
+	ReadAccessCheck();
+	return CommandLists[ID];
+}
+
+int32 FRenderTraceProvider::GetNumCommandListTimelines() const
+{
+	ReadAccessCheck();
+	return CommandListTimelines.Num();
+}
+
+void FRenderTraceProvider::ReadCommandListTimelines(TFunctionRef<void(uint32, const TEventTimeline&)> Callback) const
+{
+	ReadAccessCheck();
+	for (uint32 i = 0; i < (uint32)CommandListTimelines.Num(); ++i)
+	{
+		Callback(i, *CommandListTimelines[i]);
+	}
+}
+
+TPair<uint32, FCommandListInstance&> FRenderTraceProvider::AddCommandList(uint64 AppID, double Timestamp, ECommandListType Type)
+{
+	EditAccessCheck();
+	const uint32 ID = CommandLists.Num();
+	FCommandListInstance& NewCmdList = CommandLists.Emplace_GetRef(AppID, Timestamp, Type);
+	return TPair<uint32, FCommandListInstance&>(ID, NewCmdList);
+}
+
+FCommandListInstance& FRenderTraceProvider::EditCommandList(uint32 ID)
+{
+	EditAccessCheck();
+	return CommandLists[ID];
+}
+
+void FRenderTraceProvider::EnumerateCommandListsForEdit(TFunctionRef<void(uint32, FCommandListInstance&)> Callback)
+{
+	EditAccessCheck();
+	for (uint32 i = 0; i < (uint32)CommandLists.Num(); ++i)
+	{
+		Callback(i, CommandLists[i]);
+	}
+}
+
+TPair<uint32, FRenderTraceProvider::TEventTimeline&> FRenderTraceProvider::AddCommandListTimeline()
+{
+	EditAccessCheck();
+	++TimelinesModCount;
+	uint32 Index = CommandListTimelines.Num();
+	TUniquePtr<TEventTimeline>& Timeline = CommandListTimelines.Emplace_GetRef(MakeUnique<TEventTimeline>(Session.GetLinearAllocator()));
+	return TPair<uint32, TEventTimeline&>(Index, *Timeline);
+}
+
+FRenderTraceProvider::TEventTimeline& FRenderTraceProvider::EditCommandListTimeline(uint32 Index)
+{
+	EditAccessCheck();
+	return *CommandListTimelines[Index];
+}
+
+int32 FRenderTraceProvider::GetNumRDGPasses() const
+{
+	ReadAccessCheck();
+	return RDGPasses.Num();
+}
+
+const FRDGPassInstance& FRenderTraceProvider::GetRDGPass(uint32 ID) const
+{
+	ReadAccessCheck();
+	return RDGPasses[ID];
+}
+
+void FRenderTraceProvider::ReadRDGTimelines(TFunctionRef<void(uint32, uint32, const TEventTimeline&)> Callback) const
+{
+	ReadAccessCheck();
+	for (uint32 i = 0; i < (uint32)RDGTimelines.Num(); ++i)
+	{
+		Callback(i, RDGTimelines[i].ThreadID, *RDGTimelines[i].Timeline);
+	}
+}
+
+const FRenderTraceProvider::TEventTimeline* FRenderTraceProvider::GetRenderThreadSubmissionTimeline() const
+{
+	ReadAccessCheck();
+	return RTSubmissionTimeline.Get();
+}
+
+TPair<uint32, FRDGPassInstance&> FRenderTraceProvider::AddRDGPass(const FString& Name, uint32 ThreadID, double Timestamp, ERDGPassType Type)
+{
+	EditAccessCheck();
+	const uint32 ID = RDGPasses.Num();
+	FRDGPassInstance& NewPass = RDGPasses.Emplace_GetRef(Name, ThreadID, Timestamp, Type);
+	return TPair<uint32, FRDGPassInstance&>(ID, NewPass);
+}
+
+FRDGPassInstance& FRenderTraceProvider::EditRDGPass(uint32 ID)
+{
+	EditAccessCheck();
+	return RDGPasses[ID];
+}
+
+FRenderTraceProvider::TEventTimeline& FRenderTraceProvider::EditRDGThreadTimeline(uint32 ThreadID)
+{
+	EditAccessCheck();
+
+	uint32* ExistingIndex = RDGThreadToTimeline.Find(ThreadID);
+	if (ExistingIndex)
+	{
+		return *RDGTimelines[*ExistingIndex].Timeline;
+	}
+
+	FThreadTimeline& NewTimeline = RDGTimelines.Emplace_GetRef(ThreadID, Session.GetLinearAllocator());
+	RDGThreadToTimeline.Add(ThreadID, RDGTimelines.Num() - 1);
+	++TimelinesModCount;
+	return *NewTimeline.Timeline;
+}
+
+FRenderTraceProvider::TEventTimeline& FRenderTraceProvider::EditRenderThreadSubmissionTimeline()
+{
+	EditAccessCheck();
+	return *RTSubmissionTimeline;
+}
+
+int32 FRenderTraceProvider::GetNumRHITranslateTasks() const
+{
+	ReadAccessCheck();
+	return RHITranslateTasks.Num();
+}
+
+const FRHITranslateTask& FRenderTraceProvider::GetRHITranslateTask(uint32 ID) const
+{
+	ReadAccessCheck();
+	return RHITranslateTasks[ID];
+}
+
+uint32 FRenderTraceProvider::FindRHITaskByPredicate(uint32 StartAtTaskID, double TimeRange, TFunctionRef<int(const FRHITranslateTask&)> Pred) const
+{
+	ReadAccessCheck();
+
+	if (StartAtTaskID >= (uint32)RHITranslateTasks.Num())
+	{
+		return INVALID_EVENT_ID;
+	}
+
+	const double StartTaskTime = RHITranslateTasks[StartAtTaskID].StartTime;
+	const double MinTime = StartTaskTime - TimeRange;
+	const double MaxTime = StartTaskTime + TimeRange;
+	int32 LeftIdx = (int32)StartAtTaskID - 1;
+	int32 RightIdx = (int32)StartAtTaskID + 1;
+	bool bLeftStop = false, bRightStop = false;
+
+	while (!bLeftStop || !bRightStop)
+	{
+		if (!bLeftStop && LeftIdx >= 0 && RHITranslateTasks[LeftIdx].StartTime >= MinTime)
+		{
+			switch (Pred(RHITranslateTasks[LeftIdx]))
+			{
+			case -1: bLeftStop = true; break;
+			case 0: --LeftIdx; break;
+			case 1: return LeftIdx;
+			}
+		}
+		else
+		{
+			bLeftStop = true;
+		}
+
+		if (!bRightStop && RightIdx < RHITranslateTasks.Num() && RHITranslateTasks[RightIdx].StartTime <= MaxTime)
+		{
+			switch (Pred(RHITranslateTasks[RightIdx]))
+			{
+			case -1: bRightStop = true; break;
+			case 0: ++RightIdx; break;
+			case 1: return RightIdx;
+			}
+		}
+		else
+		{
+			bRightStop = true;
+		}
+	}
+
+	return INVALID_EVENT_ID;
+}
+
+void FRenderTraceProvider::EnumerateRHITranslateTimelines(TFunctionRef<void(uint32, uint32, const TEventTimeline&)> Callback) const
+{
+	ReadAccessCheck();
+	for (uint32 i = 0; i < (uint32)RHITranslateTimelines.Num(); ++i)
+	{
+		Callback(i, RHITranslateTimelines[i].ThreadID, *RHITranslateTimelines[i].Timeline);
+	}
+}
+
+const FRenderTraceProvider::TEventTimeline* FRenderTraceProvider::GetRHISubmissionTimeline() const
+{
+	ReadAccessCheck();
+	return RHISubmissionTimeline.Get();
+}
+
+TPair<uint32, FRHITranslateTask&> FRenderTraceProvider::AddRHITranslateTask(uint64 AppID, ERHITranslateTaskType Type, uint32 ThreadID, double Timestamp)
+{
+	EditAccessCheck();
+	const uint32 ID = RHITranslateTasks.Num();
+	FRHITranslateTask& NewTask = RHITranslateTasks.Emplace_GetRef(AppID, Type, ThreadID, Timestamp);
+	return TPair<uint32, FRHITranslateTask&>(ID, NewTask);
+}
+
+FRHITranslateTask& FRenderTraceProvider::EditRHITranslateTask(uint32 ID)
+{
+	EditAccessCheck();
+	return RHITranslateTasks[ID];
+}
+
+FRenderTraceProvider::TEventTimeline& FRenderTraceProvider::EditRHITranslateThreadTimeline(uint32 ThreadID)
+{
+	EditAccessCheck();
+
+	uint32* ExistingIndex = RHIThreadToTimeline.Find(ThreadID);
+	if (ExistingIndex)
+	{
+		return *RHITranslateTimelines[*ExistingIndex].Timeline;
+	}
+
+	FThreadTimeline& NewTimeline = RHITranslateTimelines.Emplace_GetRef(ThreadID, Session.GetLinearAllocator());
+	RHIThreadToTimeline.Add(ThreadID, RHITranslateTimelines.Num() - 1);
+	++TimelinesModCount;
+	return *NewTimeline.Timeline;
+}
+
+FRenderTraceProvider::TEventTimeline& FRenderTraceProvider::EditRHISubmissionTimeline()
+{
+	EditAccessCheck();
+	return *RHISubmissionTimeline;
+}
+
+int32 FRenderTraceProvider::GetNumPlatformPayloads() const
+{
+	ReadAccessCheck();
+	return PlatformPayloads.Num();
+}
+
+const FPlatformPayload& FRenderTraceProvider::GetPlatformPayload(uint32 ID) const
+{
+	ReadAccessCheck();
+	return PlatformPayloads[ID];
+}
+
+TPair<uint32, FPlatformPayload&> FRenderTraceProvider::AddPlatformPayload(uint64 InAppID, double InStartTime, uint8 InPipeIdx)
+{
+	EditAccessCheck();
+	const uint32 ID = PlatformPayloads.Num();
+	FPlatformPayload& NewPayload = PlatformPayloads.Emplace_GetRef(InAppID, InStartTime, InPipeIdx);
+	return TPair<uint32, FPlatformPayload&>(ID, NewPayload);
+}
+
+FPlatformPayload& FRenderTraceProvider::EditPlatformPayload(uint32 ID)
+{
+	EditAccessCheck();
+	return PlatformPayloads[ID];
+}
+
+int32 FRenderTraceProvider::GetNumSyncPoints() const
+{
+	ReadAccessCheck();
+	return SyncPoints.Num();
+}
+
+const FSyncPoint& FRenderTraceProvider::GetSyncPoint(uint32 ID) const
+{
+	ReadAccessCheck();
+	return SyncPoints[ID];
+}
+
+TPair<uint32, FSyncPoint&> FRenderTraceProvider::AddSyncPoint(uint64 InAppID, ESyncPointType InType)
+{
+	EditAccessCheck();
+	const uint32 ID = SyncPoints.Num();
+	FSyncPoint& NewSyncPoint = SyncPoints.Emplace_GetRef(InAppID, InType);
+	return TPair<uint32, FSyncPoint&>(ID, NewSyncPoint);
+}
+
+FSyncPoint& FRenderTraceProvider::EditSyncPoint(uint32 ID)
+{
+	EditAccessCheck();
+	return SyncPoints[ID];
+}
+
+int32 FRenderTraceProvider::GetNumSubmissionBatches() const
+{
+	ReadAccessCheck();
+	return SubmissionBatches.Num();
+}
+
+const FSubmissionBatch& FRenderTraceProvider::GetSubmissionBatch(uint32 ID) const
+{
+	ReadAccessCheck();
+	return SubmissionBatches[ID];
+}
+
+const FRenderTraceProvider::TEventTimeline* FRenderTraceProvider::GetSubmissionQueueTimeline() const
+{
+	ReadAccessCheck();
+	return SubmissionQueueTimeline.Get();
+}
+
+TPair<uint32, FSubmissionBatch&> FRenderTraceProvider::AddSubmissionBatch(uint32 ThreadID, double Timestamp)
+{
+	EditAccessCheck();
+	const uint32 ID = SubmissionBatches.Num();
+	FSubmissionBatch& NewBatch = SubmissionBatches.Emplace_GetRef(ThreadID, Timestamp);
+	return TPair<uint32, FSubmissionBatch&>(ID, NewBatch);
+}
+
+FSubmissionBatch& FRenderTraceProvider::EditSubmissionBatch(uint32 ID)
+{
+	EditAccessCheck();
+	return SubmissionBatches[ID];
+}
+
+FRenderTraceProvider::TEventTimeline& FRenderTraceProvider::EditSubmissionQueueTimeline()
+{
+	EditAccessCheck();
+	return *SubmissionQueueTimeline;
+}
+
+int32 FRenderTraceProvider::GetNumInterruptWakeUps() const
+{
+	ReadAccessCheck();
+	return InterruptWakeUps.Num();
+}
+
+const FInterruptWakeUp& FRenderTraceProvider::GetInterruptWakeUp(uint32 ID) const
+{
+	ReadAccessCheck();
+	return InterruptWakeUps[ID];
+}
+
+const FRenderTraceProvider::TEventTimeline* FRenderTraceProvider::GetInterruptTimeline() const
+{
+	ReadAccessCheck();
+	return InterruptTimeline.Get();
+}
+
+TPair<uint32, FInterruptWakeUp&> FRenderTraceProvider::AddInterruptWakeUp(uint32 ThreadID, double Timestamp)
+{
+	EditAccessCheck();
+	const uint32 ID = InterruptWakeUps.Num();
+	FInterruptWakeUp& NewWakeUp = InterruptWakeUps.Emplace_GetRef(ThreadID, Timestamp);
+	return TPair<uint32, FInterruptWakeUp&>(ID, NewWakeUp);
+}
+
+FInterruptWakeUp& FRenderTraceProvider::EditInterruptWakeUp(uint32 ID)
+{
+	EditAccessCheck();
+	return InterruptWakeUps[ID];
+}
+
+FRenderTraceProvider::TEventTimeline& FRenderTraceProvider::EditInterruptTimeline()
+{
+	EditAccessCheck();
+	return *InterruptTimeline;
+}
+
+} //namespace RenderTraceInsights
+} //namespace UE
diff --git a/Engine/Plugins/RenderTraceInsights/Source/RenderTraceInsights/Private/RenderTraceProvider.h b/Engine/Plugins/RenderTraceInsights/Source/RenderTraceInsights/Private/RenderTraceProvider.h
new file mode 100644
--- /dev/null
+++ b/Engine/Plugins/RenderTraceInsights/Source/RenderTraceInsights/Private/RenderTraceProvider.h
@@ -0,0 +1,152 @@
+// Copyright Epic Games, Inc. All Rights Reserved.
+
+#pragma once
+
+#include "IRenderTraceProvider.h"
+#include "InsightsCore/Common/SimpleRtti.h"
+#include "Trace/Analyzer.h"
+#include "Common/ProviderLock.h"
+
+namespace TraceServices { class IAnalysisSession; }
+
+namespace UE
+{
+namespace RenderTraceInsights
+{
+
+extern thread_local TraceServices::FProviderLock::FThreadLocalState GRenderTraceProviderLockState;
+
+class FRenderTraceProvider final : public IEditableRenderTraceProvider
+{
+public:
+	UE_NONCOPYABLE(FRenderTraceProvider)
+
+	static FName ProviderName;
+
+	FRenderTraceProvider(TraceServices::IAnalysisSession& InSession);
+
+	// IProvider methods.
+	void BeginRead() const override { Lock.BeginRead(GRenderTraceProviderLockState); }
+	void EndRead() const override { Lock.EndRead(GRenderTraceProviderLockState); }
+	void ReadAccessCheck() const override { Lock.ReadAccessCheck(GRenderTraceProviderLockState); }
+
+	// IEditableProvider methods.
+	void BeginEdit() const override { Lock.BeginWrite(GRenderTraceProviderLockState); }
+	void EndEdit() const override { Lock.EndWrite(GRenderTraceProviderLockState); }
+	void EditAccessCheck() const override { Lock.WriteAccessCheck(GRenderTraceProviderLockState); }
+
+	/////////////////////////////////////////
+	// IRenderTraceProvider methods.
+
+	uint64 GetTimelinesModCount() const override { return TimelinesModCount; }
+
+	int32 GetNumCommandLists() const override;
+	const FCommandListInstance& GetCommandList(uint32 ID) const override;
+	int32 GetNumCommandListTimelines() const override;
+	void ReadCommandListTimelines(TFunctionRef<void(uint32, const TEventTimeline&)> Callback) const override;
+
+	int32 GetNumRDGPasses() const override;
+	const FRDGPassInstance& GetRDGPass(uint32 ID) const override;
+	void ReadRDGTimelines(TFunctionRef<void(uint32, uint32, const TEventTimeline&)> Callback) const override;
+	const TEventTimeline* GetRenderThreadSubmissionTimeline() const override;
+
+	int32 GetNumRHITranslateTasks() const override;
+	const FRHITranslateTask& GetRHITranslateTask(uint32 ID) const override;
+	uint32 FindRHITaskByPredicate(uint32 StartAtTaskID, double TimeRange, TFunctionRef<int(const FRHITranslateTask&)> Pred) const override;
+	void EnumerateRHITranslateTimelines(TFunctionRef<void(uint32, uint32, const TEventTimeline&)> Callback) const override;
+	const TEventTimeline* GetRHISubmissionTimeline() const override;
+
+	int32 GetNumPlatformPayloads() const override;
+	const FPlatformPayload& GetPlatformPayload(uint32 ID) const override;
+
+	int32 GetNumSyncPoints() const override;
+	const FSyncPoint& GetSyncPoint(uint32 ID) const override;
+
+	int32 GetNumSubmissionBatches() const override;
+	const FSubmissionBatch& GetSubmissionBatch(uint32 ID) const override;
+	const TEventTimeline* GetSubmissionQueueTimeline() const override;
+
+	int32 GetNumInterruptWakeUps() const override;
+	const FInterruptWakeUp& GetInterruptWakeUp(uint32 ID) const override;
+	const TEventTimeline* GetInterruptTimeline() const override;
+
+	/////////////////////////////////////////
+	// IEditableRenderTraceProvider methods.
+
+	TPair<uint32, FCommandListInstance&> AddCommandList(uint64 AppID, double Timestamp, ECommandListType Type) override;
+	FCommandListInstance& EditCommandList(uint32 ID) override;
+	void EnumerateCommandListsForEdit(TFunctionRef<void(uint32, FCommandListInstance&)> Callback) override;
+	TPair<uint32, TEventTimeline&> AddCommandListTimeline() override;
+	TEventTimeline& EditCommandListTimeline(uint32 Index) override;
+
+	TPair<uint32, FRDGPassInstance&> AddRDGPass(const FString& Name, uint32 ThreadID, double Timestamp, ERDGPassType Type) override;
+	FRDGPassInstance& EditRDGPass(uint32 ID) override;
+	TEventTimeline& EditRDGThreadTimeline(uint32 ThreadID) override;
+	TEventTimeline& EditRenderThreadSubmissionTimeline() override;
+
+	TPair<uint32, FRHITranslateTask&> AddRHITranslateTask(uint64 AppID, ERHITranslateTaskType Type, uint32 ThreadID, double Timestamp) override;
+	FRHITranslateTask& EditRHITranslateTask(uint32 ID) override;
+	TEventTimeline& EditRHITranslateThreadTimeline(uint32 ThreadID) override;
+	TEventTimeline& EditRHISubmissionTimeline() override;
+
+	TPair<uint32, FPlatformPayload&> AddPlatformPayload(uint64 InAppID, double InStartTime, uint8 InPipeIdx) override;
+	FPlatformPayload& EditPlatformPayload(uint32 ID) override;
+
+	TPair<uint32, FSyncPoint&> AddSyncPoint(uint64 InAppID, ESyncPointType InType) override;
+	FSyncPoint& EditSyncPoint(uint32 ID) override;
+
+	TPair<uint32, FSubmissionBatch&> AddSubmissionBatch(uint32 ThreadID, double Timestamp) override;
+	FSubmissionBatch& EditSubmissionBatch(uint32 ID) override;
+	TEventTimeline& EditSubmissionQueueTimeline() override;
+
+	TPair<uint32, FInterruptWakeUp&> AddInterruptWakeUp(uint32 ThreadID, double Timestamp) override;
+	FInterruptWakeUp& EditInterruptWakeUp(uint32 ID) override;
+	TEventTimeline& EditInterruptTimeline() override;
+
+private:
+	struct FThreadTimeline
+	{
+		FThreadTimeline(uint32 InThreadID, TraceServices::ILinearAllocator& InAllocator)
+			: ThreadID(InThreadID)
+			, Timeline(MakeUnique<TEventTimeline>(InAllocator))
+		{}
+
+		uint32 ThreadID;
+		TUniquePtr<TEventTimeline> Timeline;
+	};
+
+	TraceServices::IAnalysisSession& Session;
+
+	TArray<FCommandListInstance> CommandLists;
+	TArray<TUniquePtr<TEventTimeline>> CommandListTimelines;
+
+	TArray<FRDGPassInstance> RDGPasses;
+	TArray<FThreadTimeline> RDGTimelines;
+	TMap<uint32, uint32> RDGThreadToTimeline;
+	TUniquePtr<TEventTimeline> RTSubmissionTimeline;
+
+	TArray<FRHITranslateTask> RHITranslateTasks;
+	TArray<FThreadTimeline> RHITranslateTimelines;
+	TMap<uint32, uint32> RHIThreadToTimeline;
+	TUniquePtr<TEventTimeline> RHISubmissionTimeline;
+
+	TArray<FPlatformPayload> PlatformPayloads;
+	TArray<FSyncPoint> SyncPoints;
+
+	TArray<FSubmissionBatch> SubmissionBatches;
+	TUniquePtr<TEventTimeline> SubmissionQueueTimeline;
+
+	TArray<FInterruptWakeUp> InterruptWakeUps;
+	TUniquePtr<TEventTimeline> InterruptTimeline;
+
+	mutable TraceServices::FProviderLock Lock;
+
+	// This is incremented whenever a timeline is added to let FRenderTraceTimingViewSession know that it needs to add
+	// the corresponding tracks. We don't track other modifications right now, such as adding command lists, RDG passes etc.
+	// We start at 1 because FRenderTraceTimingViewSession starts its cached value at 0, so this was we always trigger an
+	// update on the first tick.
+	uint64 TimelinesModCount = 1;
+};
+
+} //namespace RenderTraceInsights
+} //namespace UE
diff --git a/Engine/Plugins/RenderTraceInsights/Source/RenderTraceInsights/Private/RenderTraceTimingViewExtender.cpp b/Engine/Plugins/RenderTraceInsights/Source/RenderTraceInsights/Private/RenderTraceTimingViewExtender.cpp
new file mode 100644
--- /dev/null
+++ b/Engine/Plugins/RenderTraceInsights/Source/RenderTraceInsights/Private/RenderTraceTimingViewExtender.cpp
@@ -0,0 +1,72 @@
+// Copyright Epic Games, Inc. All Rights Reserved.
+
+#include "RenderTraceTimingViewExtender.h"
+#include "Insights/ITimingViewSession.h"
+#include "Insights/IUnrealInsightsModule.h"
+
+#define LOCTEXT_NAMESPACE "RenderTraceTimingViewExtender"
+
+namespace UE
+{
+namespace RenderTraceInsights
+{
+
+void FRenderTraceTimingViewExtender::OnBeginSession(UE::Insights::Timing::ITimingViewSession& InSession)
+{
+	if (InSession.GetName() != FInsightsManagerTabs::TimingProfilerTabId)
+	{
+		return;
+	}
+
+	FPerSessionData* PerSessionData = PerSessionDataMap.Find(&InSession);
+	if(PerSessionData == nullptr)
+	{
+		PerSessionData = &PerSessionDataMap.Add(&InSession);
+		PerSessionData->SharedData = MakeShared<FRenderTraceTimingViewSession>();
+
+		PerSessionData->SharedData->OnBeginSession(InSession);
+	}
+	else
+	{
+		PerSessionData->SharedData->OnBeginSession(InSession);
+	}
+}
+
+void FRenderTraceTimingViewExtender::OnEndSession(UE::Insights::Timing::ITimingViewSession& InSession)
+{
+	if (InSession.GetName() != FInsightsManagerTabs::TimingProfilerTabId)
+	{
+		return;
+	}
+
+	FPerSessionData* PerSessionData = PerSessionDataMap.Find(&InSession);
+	if(PerSessionData != nullptr)
+	{
+		PerSessionData->SharedData->OnEndSession(InSession);
+	}
+
+	PerSessionDataMap.Remove(&InSession);
+}
+
+void FRenderTraceTimingViewExtender::Tick(UE::Insights::Timing::ITimingViewSession& InSession, const TraceServices::IAnalysisSession& InAnalysisSession)
+{
+	FPerSessionData* PerSessionData = PerSessionDataMap.Find(&InSession);
+	if(PerSessionData != nullptr)
+	{
+		PerSessionData->SharedData->Tick(InSession, InAnalysisSession);
+	}
+}
+
+void FRenderTraceTimingViewExtender::ExtendFilterMenu(UE::Insights::Timing::ITimingViewSession& InSession, FMenuBuilder& InMenuBuilder)
+{
+	FPerSessionData* PerSessionData = PerSessionDataMap.Find(&InSession);
+	if (PerSessionData != nullptr)
+	{
+		PerSessionData->SharedData->ExtendFilterMenu(InMenuBuilder);
+	}
+}
+
+} //namespace RenderTraceInsights
+} //namespace UE
+
+#undef LOCTEXT_NAMESPACE
diff --git a/Engine/Plugins/RenderTraceInsights/Source/RenderTraceInsights/Private/RenderTraceTimingViewExtender.h b/Engine/Plugins/RenderTraceInsights/Source/RenderTraceInsights/Private/RenderTraceTimingViewExtender.h
new file mode 100644
--- /dev/null
+++ b/Engine/Plugins/RenderTraceInsights/Source/RenderTraceInsights/Private/RenderTraceTimingViewExtender.h
@@ -0,0 +1,38 @@
+// Copyright Epic Games, Inc. All Rights Reserved.
+
+#pragma once
+
+#include "Insights/ITimingViewExtender.h"
+#include "RenderTraceTimingViewSession.h"
+
+namespace UE::Insights::Timing { class ITimingViewSession; }
+namespace TraceServices { class IAnalysisSession; }
+class FMenuBuilder;
+
+namespace UE
+{
+namespace RenderTraceInsights
+{
+
+class FRenderTraceTimingViewExtender : public UE::Insights::Timing::ITimingViewExtender
+{
+public:
+	//~ Begin UE::Insights::Timing::ITimingViewExtender interface
+	virtual void OnBeginSession(UE::Insights::Timing::ITimingViewSession& InSession) override;
+	virtual void OnEndSession(UE::Insights::Timing::ITimingViewSession& InSession) override;
+	virtual void Tick(UE::Insights::Timing::ITimingViewSession& InSession, const TraceServices::IAnalysisSession& InAnalysisSession) override;
+	virtual void ExtendFilterMenu(UE::Insights::Timing::ITimingViewSession& InSession, FMenuBuilder& InMenuBuilder) override;
+	//~ End UE::Insights::Timing::ITimingViewExtender interface
+
+private:
+	struct FPerSessionData
+	{
+		TSharedPtr<FRenderTraceTimingViewSession> SharedData;
+	};
+
+	// The data we host per-session
+	TMap<UE::Insights::Timing::ITimingViewSession*, FPerSessionData> PerSessionDataMap;
+};
+
+} //namespace RenderTraceInsights
+} //namespace UE
diff --git a/Engine/Plugins/RenderTraceInsights/Source/RenderTraceInsights/Private/RenderTraceTimingViewSession.cpp b/Engine/Plugins/RenderTraceInsights/Source/RenderTraceInsights/Private/RenderTraceTimingViewSession.cpp
new file mode 100644
--- /dev/null
+++ b/Engine/Plugins/RenderTraceInsights/Source/RenderTraceInsights/Private/RenderTraceTimingViewSession.cpp
@@ -0,0 +1,385 @@
+// Copyright Epic Games, Inc. All Rights Reserved.
+
+#include "RenderTraceTimingViewSession.h"
+#include "IRenderTraceProvider.h"
+#include "RenderTraceTrack.h"
+#include "Insights/ITimingViewSession.h"
+#include "TraceServices/Model/Threads.h"
+#include "Insights/ViewModels/TimingEvent.h"
+#include "Insights/TimingProfiler/Tracks/ThreadTimingTrack.h"
+#include "Common/ProviderLock.h"
+#include "Framework/MultiBox/MultiBoxBuilder.h"
+#include "Framework/Commands/UICommandInfo.h"
+
+#define LOCTEXT_NAMESPACE "RenderTraceTimingViewSession"
+
+namespace UE
+{
+namespace RenderTraceInsights
+{
+
+void FRenderTraceTimingViewSession::ClearState()
+{
+	RDGPassTracks.Empty();
+	RDGThreadToTrackIndex.Empty();
+
+	CommandListTracks.Empty();
+
+	RHIThreadTracks.Empty();
+	RHIThreadToTrackIndex.Empty();
+
+	SubmissionQueueTrack.Reset();
+	InterruptTrack.Reset();
+
+	CpuThreadTrackCache.Empty();
+
+	TimingViewSession->OnSelectedEventChanged().RemoveAll(this);
+
+	AnalysisSession = nullptr;
+	RenderTraceProvider = nullptr;
+	LastTimelinesModCount = 0;
+}
+
+void FRenderTraceTimingViewSession::OnBeginSession(UE::Insights::Timing::ITimingViewSession& InTimingViewSession)
+{
+	TimingViewSession = &InTimingViewSession;
+
+	ClearState();
+
+	TimingViewSession->OnSelectedEventChanged().AddSP(SharedThis(this), &FRenderTraceTimingViewSession::OnEventSelected);
+}
+
+void FRenderTraceTimingViewSession::OnEndSession(UE::Insights::Timing::ITimingViewSession& InTimingViewSession)
+{
+	TimingViewSession = &InTimingViewSession;
+
+	ClearState();
+
+	TimingViewSession = nullptr;
+}
+
+void FRenderTraceTimingViewSession::Tick(UE::Insights::Timing::ITimingViewSession& InTimingViewSession, const TraceServices::IAnalysisSession& InAnalysisSession)
+{
+	TimingViewSession = &InTimingViewSession;
+	AnalysisSession = &InAnalysisSession;
+
+	RenderTraceProvider = ReadRenderTraceProvider(*AnalysisSession);
+	if (!RenderTraceProvider)
+	{
+		return;
+	}
+
+	TraceServices::FProviderReadScopeLock ProviderReadScope(*RenderTraceProvider);
+	if (RenderTraceProvider->GetTimelinesModCount() == LastTimelinesModCount)
+	{
+		// No changes since last time we were called, nothing to do.
+		return;
+	}
+
+	LastTimelinesModCount = RenderTraceProvider->GetTimelinesModCount();
+	
+	TraceServices::FAnalysisSessionReadScope SessionReadScope(GetAnalysisSession());
+
+	const TraceServices::IThreadProvider& ThreadProvider = TraceServices::ReadThreadProvider(*AnalysisSession);
+	bool bTrackAdded = false;
+
+	static constexpr int32 RDGTrackOrderStart = FTimingTrackOrder::Last + FTimingTrackOrder::GroupRange;
+	static constexpr int32 CmdListTrackOrderStart = RDGTrackOrderStart + 1024;
+	static constexpr int32 RHIThreadTrackOrderStart = CmdListTrackOrderStart + 1024;
+	static constexpr int32 SubmissionQueueTrackOrderStart = RHIThreadTrackOrderStart + 1024;
+	static constexpr int32 InterruptTrackOrderStart = SubmissionQueueTrackOrderStart + 1024;
+
+	RenderTraceProvider->ReadRDGTimelines([this, &ThreadProvider, &bTrackAdded](uint32 Index, uint32 ThreadID, const IRenderTraceProvider::TEventTimeline& Timeline)
+	{
+		if (Index >= (uint32)RDGPassTracks.Num())
+		{
+			RDGPassTracks.SetNum(Index + 1);
+		}
+
+		if (RDGPassTracks[Index].IsValid())
+		{
+			return;
+		}
+
+		const int32 TrackOrder = RDGTrackOrderStart + Index;
+		const TCHAR* ThreadName = ThreadProvider.GetThreadName(ThreadID);
+		const FString TrackName = FString::Printf(TEXT("RDG Pass Exec %s"), ThreadName);
+
+		TSharedPtr<FRDGThreadTrack> Track = MakeShared<FRDGThreadTrack>(*this, TrackName, &Timeline);
+		Track->SetVisibilityFlag(bTracksVisible);
+		Track->SetOrder(TrackOrder);
+		RDGPassTracks[Index] = Track;
+
+		RDGThreadToTrackIndex.Add(ThreadID, Index);
+
+		if (FCString::Strncmp(ThreadName, TEXT("RenderThread"), 12) == 0)
+		{
+			// Add a render thread submission child track to the render thread track.
+			TSharedRef<FRenderThreadSubmissionTrack> SubmissionTrack = MakeShared<FRenderThreadSubmissionTrack>(*this, TEXT("RT Submission"), RenderTraceProvider->GetRenderThreadSubmissionTimeline());
+			SubmissionTrack->SetParentTrack(Track);
+			Track->AddChildTrack(SubmissionTrack);
+		}
+
+		TimingViewSession->AddScrollableTrack(Track);
+		bTrackAdded = true;
+	});
+
+	RenderTraceProvider->ReadCommandListTimelines([this, &bTrackAdded](uint32 Index, const IRenderTraceProvider::TEventTimeline& Timeline)
+	{
+		if (Index >= (uint32)CommandListTracks.Num())
+		{
+			CommandListTracks.SetNum(Index + 1);
+		}
+
+		if (CommandListTracks[Index].IsValid())
+		{
+			return;
+		}
+
+		const int32 TrackOrder = CmdListTrackOrderStart + Index;
+		const FString TrackName = FString::Printf(TEXT("RHI CmdList Lane %u"), Index);
+
+		TSharedPtr<FCommandListTrack> Track = MakeShared<FCommandListTrack>(*this, TrackName, &Timeline);
+		Track->SetVisibilityFlag(bTracksVisible);
+		Track->SetOrder(TrackOrder);
+		CommandListTracks[Index] = Track;
+
+		TimingViewSession->AddScrollableTrack(Track);
+		bTrackAdded = true;
+	});
+
+	RenderTraceProvider->EnumerateRHITranslateTimelines([this, &ThreadProvider, &bTrackAdded](uint32 Index, uint32 ThreadID, const IRenderTraceProvider::TEventTimeline& Timeline)
+	{
+		if (Index >= (uint32)RHIThreadTracks.Num())
+		{
+			RHIThreadTracks.SetNum(Index + 1);
+		}
+
+		if (RHIThreadTracks[Index].IsValid())
+		{
+			return;
+		}
+
+		const int32 TrackOrder = RHIThreadTrackOrderStart + Index;
+		const TCHAR* ThreadName = ThreadProvider.GetThreadName(ThreadID);
+		const FString TrackName = FString::Printf(TEXT("RHI Tasks %s"), ThreadName);
+
+		TSharedPtr<FRHIThreadTrack> Track = MakeShared<FRHIThreadTrack>(*this, TrackName, &Timeline);
+		Track->SetVisibilityFlag(bTracksVisible);
+		Track->SetOrder(TrackOrder);
+		RHIThreadTracks[Index] = Track;
+
+		RHIThreadToTrackIndex.Add(ThreadID, Index);
+
+		if (FCString::Strncmp(ThreadName, TEXT("RHIThread"), 9) == 0)
+		{
+			// Add an RHI submission child track to the RHI thread track.
+			TSharedRef<FRHISubmissionTrack> SubmissionTrack = MakeShared<FRHISubmissionTrack>(*this, TEXT("RHI Submission"), RenderTraceProvider->GetRHISubmissionTimeline());
+			SubmissionTrack->SetParentTrack(Track);
+			Track->AddChildTrack(SubmissionTrack);
+		}
+
+		TimingViewSession->AddScrollableTrack(Track);
+		bTrackAdded = true;
+	});
+
+	if (!SubmissionQueueTrack.IsValid())
+	{
+		const IRenderTraceProvider::TEventTimeline* Timeline = RenderTraceProvider->GetSubmissionQueueTimeline();
+		if (Timeline)
+		{
+			SubmissionQueueTrack = MakeShared<FSubmissionQueueTrack>(*this, TEXT("Submission Queue"), Timeline);
+			SubmissionQueueTrack->SetVisibilityFlag(bTracksVisible);
+			SubmissionQueueTrack->SetOrder(SubmissionQueueTrackOrderStart);
+
+			TimingViewSession->AddScrollableTrack(SubmissionQueueTrack);
+			bTrackAdded = true;
+		}
+	}
+
+	if (!InterruptTrack.IsValid())
+	{
+		const IRenderTraceProvider::TEventTimeline* Timeline = RenderTraceProvider->GetInterruptTimeline();
+		if (Timeline)
+		{
+			InterruptTrack = MakeShared<FInterruptTrack>(*this, TEXT("Interrupt"), Timeline);
+			InterruptTrack->SetVisibilityFlag(bTracksVisible);
+			InterruptTrack->SetOrder(InterruptTrackOrderStart);
+
+			TimingViewSession->AddScrollableTrack(InterruptTrack);
+			bTrackAdded = true;
+		}
+	}
+
+	if (bTrackAdded)
+	{
+		TimingViewSession->InvalidateScrollableTracksOrder();
+	}
+}
+
+void FRenderTraceTimingViewSession::ExtendFilterMenu(FMenuBuilder& InMenuBuilder)
+{
+	InMenuBuilder.BeginSection("RenderTraceTracks", LOCTEXT("RenderTraceHeader", "RenderTrace"));
+	{
+		InMenuBuilder.AddMenuEntry(
+			LOCTEXT("RenderTraceTimingTracks", "RenderTrace Tracks"),
+			LOCTEXT("RenderTraceTimingTracks_Tooltip", "Show/hide the RenderTrace tracks"),
+			FSlateIcon(),
+			FUIAction(
+				// CreateRaw seems to be OK here because of how the UI lifetime is managed. A shared pointer would be
+				// safer, but I don't want to introduce a cyclical dependency and this pattern is used in all other
+				// view session extenders anyway.
+				FExecuteAction::CreateRaw(this, &FRenderTraceTimingViewSession::ToggleRenderTraceTrack),
+				FCanExecuteAction(),
+				FIsActionChecked::CreateLambda([this]() { return bTracksVisible; })
+				),
+			NAME_None,
+			EUserInterfaceActionType::ToggleButton
+		);
+	}
+	InMenuBuilder.EndSection();
+}
+
+TSharedPtr<const FRDGThreadTrack> FRenderTraceTimingViewSession::GetRDGThreadTrack(uint32 ThreadID) const
+{
+	const uint32* TrackIndex = RDGThreadToTrackIndex.Find(ThreadID);
+	if (!TrackIndex)
+	{
+		return nullptr;
+	}
+
+	return RDGPassTracks[*TrackIndex];
+}
+
+TSharedPtr<const FRHIThreadTrack> FRenderTraceTimingViewSession::GetRHIThreadTrack(uint32 ThreadID) const
+{
+	const uint32* TrackIndex = RHIThreadToTrackIndex.Find(ThreadID);
+	if (!TrackIndex)
+	{
+		return nullptr;
+	}
+
+	return RHIThreadTracks[*TrackIndex];
+}
+
+void FRenderTraceTimingViewSession::ToggleRenderTraceTrack()
+{
+	bTracksVisible = !bTracksVisible;
+
+	for (TSharedPtr<FRDGThreadTrack>& Track : RDGPassTracks)
+	{
+		if (Track.IsValid())
+		{
+			Track->SetVisibilityFlag(bTracksVisible);
+		}
+	}
+
+	for (TSharedPtr<FCommandListTrack>& Track : CommandListTracks)
+	{
+		if (Track.IsValid())
+		{
+			Track->SetVisibilityFlag(bTracksVisible);
+		}
+	}
+
+	for (TSharedPtr<FRHIThreadTrack>& Track : RHIThreadTracks)
+	{
+		if (Track.IsValid())
+		{
+			Track->SetVisibilityFlag(bTracksVisible);
+		}
+	}
+
+	if (SubmissionQueueTrack.IsValid())
+	{
+		SubmissionQueueTrack->SetVisibilityFlag(bTracksVisible);
+	}
+
+	if (InterruptTrack.IsValid())
+	{
+		InterruptTrack->SetVisibilityFlag(bTracksVisible);
+	}
+}
+
+void FRenderTraceTimingViewSession::OnEventSelected(TSharedPtr<const ITimingEvent> InSelectedEvent)
+{
+	// This is called by the global OnSelectedEventChanged() event. We need to catch that so we can do global operations
+	// such as clearing dependency arrows when the selection is cleared or changed to something that doesn't belong to us.
+	// After we've done this we can unpack and route the event to our tracks, so they don't need to use the
+	// FBaseTimingTrack::OnEventSelected override.
+
+	if (!TimingViewSession || !RenderTraceProvider)
+	{
+		return;
+	}
+
+	// Lock the provider for reading for convenience, since all the OnRenderTraceEventSelected() implementations will read from it.
+	TraceServices::FProviderReadScopeLock ProviderReadScope(*RenderTraceProvider);
+
+	// Clear any dependency arrows that might currently exist.
+	TimingViewSession->EditCurrentRelations().RemoveAll([](TUniquePtr<ITimingEventRelation>& Item)
+	{
+		return Item->Is<FRenderTraceRelation>();
+	});
+
+	// RenderTrace tracks use FTimingEvent to represent events, so ignore everything else.
+	if (!InSelectedEvent.IsValid() || !InSelectedEvent->Is<FTimingEvent>())
+	{
+		return;
+	}
+
+	// We only care about our own tracks.
+	TSharedRef<const FBaseTimingTrack> BaseTrack = InSelectedEvent->GetTrack();
+	if (!BaseTrack->Is<FRenderTraceTrack>())
+	{
+		return;
+	}
+
+	// Route the event to the corresponding track.
+	const FRenderTraceTrack& Track = BaseTrack->As<const FRenderTraceTrack>();
+	const FTimingEvent& Event = InSelectedEvent->As<FTimingEvent>();
+	Track.OnRenderTraceEventSelected(Event);
+}
+
+TSharedPtr<const FBaseTimingTrack> FRenderTraceTimingViewSession::GetCpuThreadTrack(uint32 ThreadID) const
+{
+	if (!TimingViewSession)
+	{
+		return nullptr;
+	}
+
+	// Check cache first.
+	if (TWeakPtr<FBaseTimingTrack>* Cached = CpuThreadTrackCache.Find(ThreadID))
+	{
+		TSharedPtr<FBaseTimingTrack> Pinned = Cached->Pin();
+		if (Pinned)
+		{
+			return Pinned;
+		}
+
+		// Stale entry, remove it.
+		CpuThreadTrackCache.Remove(ThreadID);
+	}
+
+	// Cache all the thread timing tracks we can find.
+	TSharedPtr<FBaseTimingTrack> FoundTrack;
+	TimingViewSession->EnumerateTracks([this, &FoundTrack, ThreadID](TSharedPtr<FBaseTimingTrack> Track)
+	{
+		if (Track->Is<UE::Insights::TimingProfiler::FThreadTimingTrack>())
+		{
+			UE::Insights::TimingProfiler::FThreadTimingTrack& ThreadTrack = Track->As<UE::Insights::TimingProfiler::FThreadTimingTrack>();
+			CpuThreadTrackCache.Add(ThreadTrack.GetThreadId(), Track);
+			if (ThreadTrack.GetThreadId() == ThreadID)
+			{
+				FoundTrack = Track;
+			}
+		}
+	});
+
+	return FoundTrack;
+}
+
+} //namespace RenderTraceInsights
+} //namespace UE
+
+#undef LOCTEXT_NAMESPACE
diff --git a/Engine/Plugins/RenderTraceInsights/Source/RenderTraceInsights/Private/RenderTraceTimingViewSession.h b/Engine/Plugins/RenderTraceInsights/Source/RenderTraceInsights/Private/RenderTraceTimingViewSession.h
new file mode 100644
--- /dev/null
+++ b/Engine/Plugins/RenderTraceInsights/Source/RenderTraceInsights/Private/RenderTraceTimingViewSession.h
@@ -0,0 +1,75 @@
+// Copyright Epic Games, Inc. All Rights Reserved.
+
+#pragma once
+
+#include "Templates/SharedPointer.h"
+
+namespace TraceServices { class IAnalysisSession; }
+namespace UE::Insights::Timing { class ITimingViewSession; }
+
+class ITimingEvent;
+class FMenuBuilder;
+class FBaseTimingTrack;
+
+namespace UE
+{
+namespace RenderTraceInsights
+{
+
+class IRenderTraceProvider;
+class FRenderTraceTrack;
+class FCommandListTrack;
+class FRDGThreadTrack;
+class FRHIThreadTrack;
+class FSubmissionQueueTrack;
+class FInterruptTrack;
+
+class FRenderTraceTimingViewSession : public TSharedFromThis<FRenderTraceTimingViewSession>
+{
+public:
+	FRenderTraceTimingViewSession() = default;
+
+	void OnBeginSession(UE::Insights::Timing::ITimingViewSession& InTimingViewSession);
+	void OnEndSession(UE::Insights::Timing::ITimingViewSession& InTimingViewSession);
+	void Tick(UE::Insights::Timing::ITimingViewSession& InTimingViewSession, const TraceServices::IAnalysisSession& InAnalysisSession);
+	void ExtendFilterMenu(FMenuBuilder& InMenuBuilder);
+
+	const TraceServices::IAnalysisSession& GetAnalysisSession() const { return *AnalysisSession; }
+	UE::Insights::Timing::ITimingViewSession* GetTimingViewSession() const { return TimingViewSession; }
+	const IRenderTraceProvider* GetRenderTraceProvider() const { return RenderTraceProvider; }
+	TSharedPtr<const FRDGThreadTrack> GetRDGThreadTrack(uint32 ThreadID) const;
+	TSharedPtr<const FCommandListTrack> GetCommandListTrack(int32 Index) const { return Index >= 0 && Index < CommandListTracks.Num() ? CommandListTracks[Index] : nullptr; }
+	TSharedPtr<const FRHIThreadTrack> GetRHIThreadTrack(uint32 ThreadID) const;
+	TSharedPtr<const FSubmissionQueueTrack> GetSubmissionTrack() const { return SubmissionQueueTrack; }
+	TSharedPtr<const FInterruptTrack> GetInterruptTrack() const { return InterruptTrack; }
+
+	TSharedPtr<const FBaseTimingTrack> GetCpuThreadTrack(uint32 ThreadID) const;
+
+private:
+	void ClearState();
+	void ToggleRenderTraceTrack();
+	void OnEventSelected(TSharedPtr<const ITimingEvent> InSelectedEvent);
+
+	const TraceServices::IAnalysisSession* AnalysisSession = nullptr;
+	UE::Insights::Timing::ITimingViewSession* TimingViewSession = nullptr;
+	const IRenderTraceProvider* RenderTraceProvider = nullptr;
+	
+	TArray<TSharedPtr<FRDGThreadTrack>> RDGPassTracks;
+	TMap<uint32, uint32> RDGThreadToTrackIndex;
+
+	TArray<TSharedPtr<FCommandListTrack>> CommandListTracks;
+
+	TArray<TSharedPtr<FRHIThreadTrack>> RHIThreadTracks;
+	TMap<uint32, uint32> RHIThreadToTrackIndex;
+
+	TSharedPtr<FSubmissionQueueTrack> SubmissionQueueTrack;
+	TSharedPtr<FInterruptTrack> InterruptTrack;
+
+	bool bTracksVisible = true;
+	uint64 LastTimelinesModCount = 0;
+
+	mutable TMap<uint32, TWeakPtr<FBaseTimingTrack>> CpuThreadTrackCache;
+};
+
+} //namespace RenderTraceInsights
+} //namespace UE
diff --git a/Engine/Plugins/RenderTraceInsights/Source/RenderTraceInsights/Private/RenderTraceTrack.cpp b/Engine/Plugins/RenderTraceInsights/Source/RenderTraceInsights/Private/RenderTraceTrack.cpp
new file mode 100644
--- /dev/null
+++ b/Engine/Plugins/RenderTraceInsights/Source/RenderTraceInsights/Private/RenderTraceTrack.cpp
@@ -0,0 +1,2497 @@
+// Copyright Epic Games, Inc. All Rights Reserved.
+
+#include "RenderTraceTrack.h"
+#include "RenderTraceProvider.h"
+#include "RenderTraceTimingViewSession.h"
+#include "RenderTraceModule.h"
+#include "InsightsCore/Common/TimeUtils.h"
+#include "InsightsCore/Common/PaintUtils.h"
+#include "Insights/ITimingViewSession.h"
+#include "Insights/ViewModels/TimingEvent.h"
+#include "Insights/ViewModels/ITimingViewDrawHelper.h"
+#include "Misc/StringBuilder.h"
+#include "HAL/PlatformApplicationMisc.h"
+#include "Common/ProviderLock.h"
+
+namespace UE {
+namespace RenderTraceInsights {
+
+INSIGHTS_IMPLEMENT_RTTI(FRenderTraceRelation)
+INSIGHTS_IMPLEMENT_RTTI(FRenderTraceTrack)
+INSIGHTS_IMPLEMENT_RTTI(FCommandListTrack)
+INSIGHTS_IMPLEMENT_RTTI(FRDGThreadTrack)
+INSIGHTS_IMPLEMENT_RTTI(FRenderThreadSubmissionTrack)
+INSIGHTS_IMPLEMENT_RTTI(FRHIThreadTrack)
+INSIGHTS_IMPLEMENT_RTTI(FRHISubmissionTrack)
+INSIGHTS_IMPLEMENT_RTTI(FSubmissionQueueTrack)
+INSIGHTS_IMPLEMENT_RTTI(FInterruptTrack)
+
+constexpr FLinearColor GRelationColorRecording(0.0f, 1.0f, 0.0f, 1.0f);
+constexpr FLinearColor GRelationColorDetach(1.0f, 1.0f, 0.0f, 1.0f);
+constexpr FLinearColor GRelationColorSubmitCmdList(1.0f, 0.647f, 0.0f, 1.0f);
+constexpr FLinearColor GRelationColorTranslate(1.0f, 0.0f, 0.0f, 1.0f);
+constexpr FLinearColor GRelationColorSplitFromTranslateJob(0.0f, 0.0f, 1.0f, 1.0f);
+constexpr FLinearColor GRelationColorSplitToTranslateJob(0.0f, 0.2f, 1.0f, 1.0f);
+constexpr FLinearColor GRelationColorSubmitRHI(1.0f, 0.0f, 1.0f, 1.0f);
+constexpr FLinearColor GRelationColorExecute(1.0f, 0.5f, 0.5f, 1.0f);
+constexpr FLinearColor GRelationColorResolveSyncPoint(0.0f, 0.5f, 0.5f, 1.0f);
+constexpr FLinearColor GRelationColorInterrupt(1.0f, 0.5f, 0.0f, 1.0f);
+constexpr FLinearColor GRelationColorExecThread(1.0f, 1.0f, 1.0f, 1.0f);
+
+static const TCHAR* GetRDGPassTypeName(ERDGPassType Type)
+{
+	switch (Type)
+	{
+	case ERDGPassType::Regular: return TEXT("Regular");
+	case ERDGPassType::Dispatch: return TEXT("Dispatch");
+	case ERDGPassType::Submit: break; // This function should never be called for submit passes.
+	}
+
+	checkNoEntry();
+	return TEXT("INVALID");
+}
+
+static const TCHAR* GetCommandListTypeName(ECommandListType Type)
+{
+	switch (Type)
+	{
+	case ECommandListType::Regular: return TEXT("Regular");
+	case ECommandListType::Immediate: return TEXT("Immediate");
+	case ECommandListType::Detached: return TEXT("Detached");
+	}
+
+	checkNoEntry();
+	return TEXT("INVALID");
+}
+
+static const TCHAR* GetPipeName(uint32 PipeIdx)
+{
+	switch (PipeIdx)
+	{
+	case 0: return TEXT("Graphics");
+	case 1: return TEXT("Async Compute");
+	case 2: return TEXT("Copy");
+	}
+
+	checkNoEntry();
+	return TEXT("INVALID");
+}
+
+static const TCHAR* GetSyncPointTypeName(ESyncPointType Type)
+{
+	switch (Type)
+	{
+	case ESyncPointType::GPU: return TEXT("GPU");
+	case ESyncPointType::GPUAndCPU: return TEXT("GPUAndCPU");
+	case ESyncPointType::Manual: return TEXT("Manual");
+	}
+
+	checkNoEntry();
+	return TEXT("INVALID");
+}
+
+static FString GetCommandListName(const FCommandListInstance& CmdList)
+{
+	return FString::Printf(TEXT("0x%llx"), CmdList.AppID);
+}
+
+static FString GetTranslateTaskName(const FRHITranslateTask& Task)
+{
+	return FString::Printf(TEXT("0x%llx"), Task.AppID);
+}
+
+static FString GetExitStatusStr(uint8 ExitStatus)
+{
+	FString ExitStatusStr;
+	if (ExitStatus & 0x01) ExitStatusStr += TEXT("Processed ");
+	if (ExitStatus & 0x02) ExitStatusStr += TEXT("Pending ");
+	if (ExitStatusStr.IsEmpty()) ExitStatusStr = TEXT("Empty");
+	return ExitStatusStr;
+}
+
+static TSharedPtr<const FBaseTimingTrack> GetFirstChildTrack(TSharedPtr<const FTimingEventsTrack> Track)
+{
+	if (!Track.IsValid())
+	{
+		return nullptr;
+	}
+
+	TArrayView<const TSharedRef<FBaseTimingTrack>> Children = Track->GetChildTracks();
+	if (Children.IsEmpty())
+	{
+		return nullptr;
+	}
+
+	return Children[0];
+}
+
+#define ADD_FLAG_TO_LIST(Field, FlagNamespace, FlagName) if (Field & FlagNamespace##_##FlagName) FlagList += TEXT(#FlagName " ")
+
+constexpr uint32 RELATED_TRANSLATE_JOB_NOT_FOUND = INVALID_EVENT_ID - 1;
+
+enum class EEventDetailLevel
+{
+	Brief,
+	Full
+};
+
+class IEventDetailsBuilder
+{
+public:
+	explicit IEventDetailsBuilder(const IRenderTraceProvider& InProvider) : Provider(InProvider) {}
+	virtual ~IEventDetailsBuilder() {}
+	virtual EEventDetailLevel GetDetailLevel() const = 0;
+	virtual void AddTitle(FStringView Title) = 0;
+	virtual void AddNameValueTextLine(FStringView Name, FStringView Value) = 0;
+	const IRenderTraceProvider& GetProvider() const { return Provider; }
+protected:
+	const IRenderTraceProvider& Provider;
+};
+
+class FTooltipEventDetailsBuilder : public IEventDetailsBuilder
+{
+public:
+	FTooltipEventDetailsBuilder(const IRenderTraceProvider& InProvider, FTooltipDrawState& InTooltip) 
+		: IEventDetailsBuilder(InProvider)
+		, Tooltip(InTooltip)
+	{
+		Tooltip.ResetContent();
+	}
+
+	~FTooltipEventDetailsBuilder() {}
+
+	EEventDetailLevel GetDetailLevel() const override { return EEventDetailLevel::Brief; }
+
+	void AddTitle(FStringView Title) override
+	{
+		Tooltip.AddTitle(Title);
+	}
+
+	void AddNameValueTextLine(FStringView Name, FStringView Value) override
+	{
+		Tooltip.AddNameValueTextLine(Name, Value);
+	}
+
+	void Finish()
+	{
+		Tooltip.UpdateLayout();
+	}
+
+private:
+	FTooltipDrawState& Tooltip;
+};
+
+class FClipboardEventDetailsBuilder : public IEventDetailsBuilder
+{
+public:
+	FClipboardEventDetailsBuilder(const IRenderTraceProvider& InProvider) : IEventDetailsBuilder(InProvider) {}
+	~FClipboardEventDetailsBuilder() {}
+
+	EEventDetailLevel GetDetailLevel() const override { return EEventDetailLevel::Full; }
+
+	void AddTitle(FStringView Title) override
+	{
+		StringBuilder << Title;
+		StringBuilder << TEXT("\n");
+	}
+
+	void AddNameValueTextLine(FStringView Name, FStringView Value) override
+	{
+		StringBuilder << Name;
+		StringBuilder << TEXT(" ");
+		StringBuilder << Value;
+		StringBuilder << TEXT("\n");
+	}
+
+	void Finish()
+	{
+		FPlatformApplicationMisc::ClipboardCopy(StringBuilder.ToString());
+	}
+
+private:
+	TStringBuilder<10240> StringBuilder;
+};
+
+template<typename ItemType> uint64 GetAppIDForItemID(const IRenderTraceProvider& Provider, uint32 ItemID);
+template<> uint64 GetAppIDForItemID<FCommandListInstance>(const IRenderTraceProvider& Provider, uint32 ItemID) { return Provider.GetCommandList(ItemID).AppID; }
+template<> uint64 GetAppIDForItemID<FRHITranslateTask>(const IRenderTraceProvider& Provider, uint32 ItemID) { return Provider.GetRHITranslateTask(ItemID).AppID; }
+
+template<typename ItemType>
+class FAppIDListBuilder
+{
+public:
+	explicit FAppIDListBuilder(IEventDetailsBuilder& InDetailsBuilder, int32 InNumItems)
+		: DetailsBuilder(InDetailsBuilder)
+		, NumItems(InNumItems)
+	{}
+	virtual ~FAppIDListBuilder() {}
+
+	void AddItem(uint64 AppID)
+	{
+		Line += FString::Printf(TEXT("0x%llx"), AppID);
+		if (ItemIdx < NumItems - 1)
+		{
+			Line += TEXT(", ");
+		}
+		++ItemIdx;
+
+		if ((ItemIdx % 8) == 0)
+		{
+			DetailsBuilder.AddNameValueTextLine(TEXTVIEW(""), Line);
+			Line = TEXT("");
+		}
+	}
+	
+	void AddEvents(TConstArrayView<FEventInterval> Events, bool bFinish)
+	{
+		for (const FEventInterval& Event : Events)
+		{
+			uint64 AppID = GetAppIDForItemID<ItemType>(DetailsBuilder.GetProvider(), Event.ItemID);
+			AddItem(AppID);
+		}
+
+		if (bFinish)
+		{
+			Finish();
+		}
+	}
+
+	void AddIDs(TConstArrayView<uint64> AppIDs, bool bFinish)
+	{
+		for (uint64 ID : AppIDs)
+		{
+			AddItem(ID);
+		}
+
+		if (bFinish)
+		{
+			Finish();
+		}
+	}
+
+	void Finish()
+	{
+		if (!Line.IsEmpty())
+		{
+			DetailsBuilder.AddNameValueTextLine(TEXTVIEW(""), Line);
+		}
+	}
+
+	static void AddAllEvents(IEventDetailsBuilder& DetailsBuilder, const TCHAR* Title, TConstArrayView<FEventInterval> Events)
+	{
+		FAppIDListBuilder<ItemType> ListBuilder(DetailsBuilder, Events.Num());
+		DetailsBuilder.AddNameValueTextLine(Title, LexToString(Events.Num()));
+		ListBuilder.AddEvents(Events, true);
+	}
+
+protected:
+	IEventDetailsBuilder& DetailsBuilder;
+	int32 NumItems;
+	int32 ItemIdx = 0;
+	FString Line;
+};
+
+FRenderTraceRelation::FRenderTraceRelation(TSharedPtr<const FBaseTimingTrack> InSourceTrack, double InSourceTime, uint32 InSourceDepth, TSharedPtr<const FBaseTimingTrack> InTargetTrack, double InTargetTime, uint32 InTargetDepth, const FLinearColor& InColor)
+	: SourceTrack(InSourceTrack)
+	, TargetTrack(InTargetTrack)
+	, SourceTime(InSourceTime)
+	, TargetTime(InTargetTime)
+	, SourceDepth(InSourceDepth)
+	, TargetDepth(InTargetDepth)
+	, Color(InColor)
+{
+}
+
+void FRenderTraceRelation::Draw(const UE::Insights::FDrawContext& DrawContext, const FTimingTrackViewport& Viewport, const ITimingViewDrawHelper& Helper, const ITimingEventRelation::EDrawFilter Filter)
+{
+	int32 LayerId = Helper.GetRelationLayerId();
+
+	TSharedPtr<const FBaseTimingTrack> SourceTrackShared = SourceTrack.Pin();
+	TSharedPtr<const FBaseTimingTrack> TargetTrackShared = TargetTrack.Pin();
+
+	if (!SourceTrackShared.IsValid() || !TargetTrackShared.IsValid() || (!SourceTrackShared->IsVisible() && !TargetTrackShared->IsVisible()))
+	{
+		return;
+	}
+
+	if (Filter == ITimingEventRelation::EDrawFilter::BetweenScrollableTracks)
+	{
+		if (SourceTrackShared->GetLocation() != ETimingTrackLocation::Scrollable ||
+			TargetTrackShared->GetLocation() != ETimingTrackLocation::Scrollable)
+		{
+			return;
+		}
+	}
+
+	if (Filter == ITimingEventRelation::EDrawFilter::BetweenDockedTracks)
+	{
+		if (SourceTrackShared->GetLocation() == ETimingTrackLocation::Scrollable &&
+			TargetTrackShared->GetLocation() == ETimingTrackLocation::Scrollable)
+		{
+			return;
+		}
+
+		LayerId = DrawContext.LayerId;
+	}
+
+	FVector2D StartPoint, EndPoint;
+
+	StartPoint.X = Viewport.TimeToSlateUnitsRounded(SourceTime);
+	EndPoint.X = Viewport.TimeToSlateUnitsRounded(TargetTime);
+	if (FMath::Max(StartPoint.X, EndPoint.X) < 0.0f || FMath::Min(StartPoint.X, EndPoint.X) > Viewport.GetWidth())
+	{
+		return;
+	}
+
+	const auto ComputeY = [&Viewport](const FBaseTimingTrack* Track, uint32 Depth) -> float {
+		const float ChildTracksHeight = Track->GetChildTracksTopHeight(Viewport.GetLayout());
+		const float LaneY = Viewport.GetLayout().GetLaneY(Depth) + Viewport.GetLayout().EventH / 2.0f;
+		return Track->GetPosY() + ChildTracksHeight + LaneY;
+	};
+
+	StartPoint.Y = ComputeY(SourceTrackShared.Get(), SourceDepth);
+	EndPoint.Y = ComputeY(TargetTrackShared.Get(), TargetDepth);
+
+	const double Distance = FVector2D::Distance(StartPoint, EndPoint);
+	constexpr double LineHeightAtStart = 4.0;
+	constexpr double LineLengthAtStart = 4.0;
+	constexpr double LineLengthAtEnd = 12.0;
+	constexpr float OutlineThickness = 5.0f;
+	constexpr float LineThickness = 3.0f;
+	constexpr double ArrowDirectionLen = 10.0;
+	constexpr double ArrowRotationAngle = 20.0;
+
+	const FVector2D StartDir(FMath::Max((EndPoint.X - StartPoint.X), 4.0 * (LineLengthAtStart + LineLengthAtEnd)), 0.0);
+	FVector2D ArrowDirection;
+
+	// This makes a heap allocation for only two points, but sadly FSlateDrawElement::MakeLines wants a TArray so we can't
+	// even add an inline allocator. That API needs to be changed to take a TArrayView.
+	TArray<FVector2D> LinePoints;
+	LinePoints.SetNum(2);
+
+	constexpr FLinearColor OutlineColor(0.0f, 0.0f, 0.0f, 1.0f);
+	const int32 OutlineLayerId = LayerId - 1;
+
+	const auto DrawLine = [&DrawContext, LayerId, OutlineLayerId, OutlineColor, Color=Color, &LinePoints](FVector2D Start, FVector2D End) {
+		LinePoints[0] = Start;
+		LinePoints[1] = End;
+		DrawContext.DrawLines(OutlineLayerId, 0.0f, 0.0f, LinePoints, OutlineColor, /*bAntialias=*/ true, OutlineThickness);
+		DrawContext.DrawLines(LayerId, 0.0f, 0.0f, LinePoints, Color, /*bAntialias=*/ true, LineThickness);
+	};
+
+	// Cap the start of the line.
+	DrawLine(StartPoint + FVector2D(0.0, -LineHeightAtStart / 2.0), StartPoint + FVector2D(0.0, LineHeightAtStart / 2.0));
+
+	constexpr double MinDistance = 1.5 * (LineLengthAtStart + LineLengthAtEnd);
+	constexpr double MaxDistance = 10000.0; // arbitrary limit to avoid stack overflow in recursive FLineBuilder::Subdivide when rendering splines
+	if (Distance > MinDistance && Distance < MaxDistance && !FMath::IsNearlyEqual(StartPoint.Y, EndPoint.Y))
+	{
+		// Curve with straight start and end sections.
+		FVector2D SplineStart(StartPoint.X + LineLengthAtStart, StartPoint.Y);
+		FVector2D SplineEnd(EndPoint.X - LineLengthAtEnd, EndPoint.Y);
+		DrawLine(StartPoint, SplineStart);
+
+		DrawContext.DrawSpline(OutlineLayerId, 0.0f, 0.0f, SplineStart, StartDir, SplineEnd, StartDir, OutlineThickness, OutlineColor);
+		DrawContext.DrawSpline(LayerId, 0.0f, 0.0f, SplineStart, StartDir, SplineEnd, StartDir, LineThickness, Color);
+
+		DrawLine(SplineEnd, EndPoint);
+
+		// The arrow head will be horizontal.
+		ArrowDirection = FVector2D(-ArrowDirectionLen, 0.0);
+	}
+	else
+	{
+		// No room for curve, use a straight line.
+		DrawLine(StartPoint, EndPoint);
+
+		// The arrow head direction will be along the line.
+		ArrowDirection = StartPoint - EndPoint;
+		ArrowDirection.Normalize();
+		ArrowDirection *= ArrowDirectionLen;
+	}
+
+	// Arrow head.
+	DrawLine(EndPoint, EndPoint + ArrowDirection.GetRotated(-ArrowRotationAngle));
+	DrawLine(EndPoint, EndPoint + ArrowDirection.GetRotated(+ArrowRotationAngle));
+}
+
+FRenderTraceTrack::FRenderTraceTrack(const FRenderTraceTimingViewSession& InSharedData, const FString& Name, const IRenderTraceProvider::TEventTimeline* InTimeline)
+	: FTimingEventsTrack(Name)
+	, SharedData(InSharedData)
+	, Timeline(InTimeline)
+{
+}
+
+const TSharedPtr<const ITimingEvent> FRenderTraceTrack::GetEvent(double InTime, double SecondsPerPixel, int32 Depth) const
+{
+	const IRenderTraceProvider* Provider = SharedData.GetRenderTraceProvider();
+	if (!Timeline || !Provider)
+	{
+		return nullptr;
+	}
+
+	TraceServices::FProviderReadScopeLock ProviderReadScope(*Provider);
+
+	const double Tolerance = 2.0 * SecondsPerPixel;
+
+	// Use the timeline object to find the event at this time and depth.
+	TSharedPtr<FTimingEvent> FoundEvent;
+	Timeline->EnumerateEvents(InTime - Tolerance, InTime + Tolerance,
+	[this, Depth, InTime, &FoundEvent](double TimelineEventStart, double TimelineEventEnd, uint32 TimelineEventDepth, const FRenderTraceEvent& TimelineEvent)
+	{
+		if (InTime < TimelineEventStart || InTime > TimelineEventEnd)
+		{
+			return TraceServices::EEventEnumerate::Continue;
+		}
+
+		// Allow the subclass to customize the find event logic.
+		FoundEvent = GetEventAtTime(InTime, Depth, TimelineEventStart, TimelineEventEnd, TimelineEventDepth, TimelineEvent);
+		if (FoundEvent.IsValid())
+		{
+			return TraceServices::EEventEnumerate::Stop;
+		}
+
+		// In most cases we can just use the timeline event at the corresponding depth.
+		if (TimelineEventDepth == Depth)
+		{
+			// We'll just use FTimingEvent since it has a uint64 field (Type) where we can pack our 32-bit item and parent IDs. Hacky but works for now.
+			const uint64 Packed = ((uint64)TimelineEvent.ParentItemID << 32) | TimelineEvent.ItemID;
+			FoundEvent = MakeShared<FTimingEvent>(SharedThis(this), TimelineEventStart, TimelineEventEnd, Depth, Packed);
+			return TraceServices::EEventEnumerate::Stop;
+		}
+
+		return TraceServices::EEventEnumerate::Continue;
+	});
+
+	return FoundEvent;
+}
+
+void FRenderTraceTrack::BuildDrawState(ITimingEventsTrackDrawStateBuilder& Builder, const ITimingTrackUpdateContext& Context)
+{
+	const IRenderTraceProvider* Provider = SharedData.GetRenderTraceProvider();
+	if (!Timeline || !Provider)
+	{
+		return;
+	}
+
+	TraceServices::FProviderReadScopeLock ProviderReadScope(*Provider);
+
+	const FTimingTrackViewport& Viewport = Context.GetViewport();
+	const double SecondsPerPixel = 1.0 / Viewport.GetScaleX();
+
+	Timeline->EnumerateEventsDownSampled(Viewport.GetStartTime(), Viewport.GetEndTime(), SecondsPerPixel, [this, &Builder, Provider](double InStartTime, double InEndTime, uint32 InDepth, const FRenderTraceEvent& InEvent)
+	{
+		// TODO: use string views here to avoid allocations. Currently not possible because some event names are built on the fly.
+		FString EventName = GetEventName(Provider, InDepth, InEvent);
+		Builder.AddEvent(InStartTime, InEndTime, InDepth, *EventName);
+		CustomizeEventDrawState(Builder, InStartTime, InEndTime, InDepth, InEvent);
+		return TraceServices::EEventEnumerate::Continue;
+	});
+}
+
+void FRenderTraceTrack::OnClipboardCopyEvent(const ITimingEvent& InSelectedEvent) const
+{
+	const IRenderTraceProvider* Provider = SharedData.GetRenderTraceProvider();
+	if (!Provider)
+	{
+		return;
+	}
+
+	TraceServices::FProviderReadScopeLock ProviderReadScope(*Provider);
+
+	const FTimingEvent& TimingEvent = InSelectedEvent.As<FTimingEvent>();
+
+	FClipboardEventDetailsBuilder Builder(*Provider);
+	PopulateEventDetails(Builder, TimingEvent);
+	Builder.Finish();
+}
+
+void FRenderTraceTrack::InitTooltip(FTooltipDrawState& Tooltip, const ITimingEvent& TooltipEvent) const
+{
+	if (!TooltipEvent.Is<FTimingEvent>())
+	{
+		return;
+	}
+
+	if (!TooltipEvent.CheckTrack(this))
+	{
+		// The event isn't meant for us, but maybe it's meant for one of our child tracks.
+		for (const TSharedRef<FBaseTimingTrack>& Track : GetChildTracks())
+		{
+			Track->InitTooltip(Tooltip, TooltipEvent);
+		}
+		return;
+	}
+
+	const IRenderTraceProvider* Provider = SharedData.GetRenderTraceProvider();
+	if (!Provider)
+	{
+		Tooltip.ResetContent();
+		Tooltip.UpdateLayout();
+		return;
+	}
+
+	TraceServices::FProviderReadScopeLock ProviderReadScope(*Provider);
+	const FTimingEvent& TimingEvent = TooltipEvent.As<FTimingEvent>();
+
+	FTooltipEventDetailsBuilder Builder(*Provider, Tooltip);
+	PopulateEventDetails(Builder, TimingEvent);
+	Builder.Finish();
+}
+
+FCommandListTrack::FCommandListTrack(const FRenderTraceTimingViewSession& InSharedData, const FString& Name, const IRenderTraceProvider::TEventTimeline* InTimeline)
+	: FRenderTraceTrack(InSharedData, Name, InTimeline)
+{
+}
+
+void FCommandListTrack::PopulateCommandListEventDetails(IEventDetailsBuilder& Builder, uint32 CmdListID) const
+{
+	const IRenderTraceProvider& Provider = Builder.GetProvider();
+	const FCommandListInstance& CmdList = Provider.GetCommandList(CmdListID);
+
+	Builder.AddTitle(FString::Printf(TEXT("Command list %s"), *GetCommandListName(CmdList)));
+	Builder.AddNameValueTextLine(TEXTVIEW("Type:"), GetCommandListTypeName(CmdList.Type));
+
+	{
+		FString FlagList;
+		ADD_FLAG_TO_LIST(CmdList.RecordingFlags, ECommandListRecordingFlag, UsesLockFence);
+		if (FlagList.IsEmpty()) FlagList = TEXT("None");
+		Builder.AddNameValueTextLine(TEXTVIEW("Flags:"), FlagList);
+	}
+
+	Builder.AddNameValueTextLine(TEXTVIEW("Creation Time:"), Insights::FormatTime(CmdList.CreateTime, UE::Insights::FTimeValue::Microsecond));
+
+	if (CmdList.FinishRecordingTime != 0.0)
+	{
+		Builder.AddNameValueTextLine(TEXTVIEW("Finish Recording Time:"), Insights::FormatTime(CmdList.FinishRecordingTime, UE::Insights::FTimeValue::Microsecond));
+	}
+
+	if (CmdList.SubmitTime != 0.0)
+	{
+		Builder.AddNameValueTextLine(TEXTVIEW("Submit Time:"), Insights::FormatTime(CmdList.SubmitTime, UE::Insights::FTimeValue::Microsecond));
+	}
+
+	Builder.AddNameValueTextLine(TEXTVIEW("Destruction Time:"), Insights::FormatTime(CmdList.DestroyTime, UE::Insights::FTimeValue::Microsecond));
+
+	if (CmdList.DetachedCmdListID != INVALID_EVENT_ID)
+	{
+		const FCommandListInstance& DetachedCmdList = Provider.GetCommandList(CmdList.DetachedCmdListID);
+		Builder.AddNameValueTextLine(TEXTVIEW("Detached To:"), GetCommandListName(DetachedCmdList));
+		Builder.AddNameValueTextLine(TEXTVIEW("Detached At:"), Insights::FormatTime(CmdList.DetachTime, UE::Insights::FTimeValue::Microsecond));
+	}
+
+	if (CmdList.SourceCmdListID != INVALID_EVENT_ID)
+	{
+		const FCommandListInstance& SourceCmdList = Provider.GetCommandList(CmdList.SourceCmdListID);
+		Builder.AddNameValueTextLine(TEXTVIEW("Detached From:"), GetCommandListName(SourceCmdList));
+		Builder.AddNameValueTextLine(TEXTVIEW("Detached At:"), Insights::FormatTime(SourceCmdList.DetachTime, UE::Insights::FTimeValue::Microsecond));
+	}
+
+	if (CmdList.DetachedCmdListID == INVALID_EVENT_ID)
+	{
+		// The source command list for detach actions is never submitted or translated, no point in showing this information.
+		if (CmdList.SubmitPassID != INVALID_EVENT_ID)
+		{
+			Builder.AddNameValueTextLine(TEXTVIEW("Submitted By:"), LexToString(CmdList.SubmitPassID));
+		}
+		else
+		{
+			Builder.AddNameValueTextLine(TEXTVIEW("NOT SUBMITTED"), TEXTVIEW(""));
+		}
+
+		if (CmdList.TranslateTaskID != INVALID_EVENT_ID)
+		{
+			const FRHITranslateTask& TranslateTask = Provider.GetRHITranslateTask(CmdList.TranslateTaskID);
+			Builder.AddNameValueTextLine(TEXTVIEW("Translated By:"), GetTranslateTaskName(TranslateTask));
+		}
+		else
+		{
+			Builder.AddNameValueTextLine(TEXTVIEW("NOT TRANSLATED"), TEXTVIEW(""));
+		}
+	}
+
+	if(!CmdList.RecordingEvents.IsEmpty())
+	{
+		Builder.AddNameValueTextLine(TEXTVIEW("Recording Events:"), LexToString(CmdList.RecordingEvents.Num()));
+		Builder.AddNameValueTextLine(TEXTVIEW("Recording Interval:"), FString::Printf(TEXT("%s -> %s"),
+			*Insights::FormatTime(CmdList.RecordingEvents[0].Start, UE::Insights::FTimeValue::Microsecond),
+			*Insights::FormatTime(CmdList.RecordingEvents.Last().End, UE::Insights::FTimeValue::Microsecond)
+		));
+
+		if (Builder.GetDetailLevel() == EEventDetailLevel::Full)
+		{
+			for (const FEventInterval& RecordingEvent : CmdList.RecordingEvents)
+			{
+				const FRDGPassInstance& Pass = Provider.GetRDGPass(RecordingEvent.ItemID);
+				Builder.AddNameValueTextLine(TEXTVIEW(""), Pass.Name);
+			}
+		}
+	}
+}
+
+void FCommandListTrack::PopulateRecordingEventDetails(IEventDetailsBuilder& Builder, uint32 PassID, double StartTime, double EndTime) const
+{
+	const IRenderTraceProvider& Provider = Builder.GetProvider();
+	const FRDGPassInstance& Pass = Provider.GetRDGPass(PassID);
+
+	Builder.AddTitle(Pass.Name);
+	Builder.AddNameValueTextLine(TEXTVIEW("Type:"), GetRDGPassTypeName(Pass.Type));
+	Builder.AddNameValueTextLine(TEXTVIEW("Recording Start:"), Insights::FormatTime(StartTime, UE::Insights::FTimeValue::Microsecond));
+	Builder.AddNameValueTextLine(TEXTVIEW("Recording Duration:"), Insights::FormatTimeAuto(EndTime - StartTime));
+}
+
+void FCommandListTrack::PopulateEventDetails(IEventDetailsBuilder& Builder, const FTimingEvent& Event) const
+{
+	const uint32 ItemID = static_cast<uint32>(Event.GetType());
+	switch (Event.GetDepth())
+	{
+	case 0:
+	{
+		PopulateCommandListEventDetails(Builder, ItemID);
+		break;
+	}
+
+	case 1:
+	{
+		PopulateRecordingEventDetails(Builder, ItemID, Event.GetStartTime(), Event.GetEndTime());
+		break;
+	}
+
+	default:
+		checkNoEntry();
+		return;
+	}
+}
+
+void FCommandListTrack::OnCommandListSelected(const FTimingEvent& InSelectedEvent) const
+{
+	UE::Insights::Timing::ITimingViewSession* TimingView = SharedData.GetTimingViewSession();
+	const IRenderTraceProvider* Provider = SharedData.GetRenderTraceProvider();
+	check(TimingView && Provider);
+
+	const uint32 CmdListID = static_cast<uint32>(InSelectedEvent.GetType());
+	const FCommandListInstance& CmdList = Provider->GetCommandList(CmdListID);
+
+	// Add arrows from all the passes which recorded to this command list.
+	for (const FEventInterval& RecordingEvent : CmdList.RecordingEvents)
+	{
+		const FRDGPassInstance& Pass = Provider->GetRDGPass(RecordingEvent.ItemID);
+		TSharedPtr<const FRDGThreadTrack> PassTrack = SharedData.GetRDGThreadTrack(Pass.ExecThreadID);
+		if (!PassTrack.IsValid())
+		{
+			continue;
+		}
+
+		TUniquePtr<ITimingEventRelation> Relation = MakeUnique<FRenderTraceRelation>(
+			PassTrack, Pass.StartTime, 0, SharedThis(this), RecordingEvent.Start, 1, GRelationColorRecording
+		);
+
+		TimingView->AddRelation(Relation);
+	}
+
+	// Add arrows to all detached command lists.
+	if (CmdList.DetachedCmdListID != INVALID_EVENT_ID)
+	{
+		const FCommandListInstance& DetachedCmdList = Provider->GetCommandList(CmdList.DetachedCmdListID);
+		TSharedPtr<const FCommandListTrack> DetachedCmdListTrack = SharedData.GetCommandListTrack(DetachedCmdList.TimelineIndex);
+		if (DetachedCmdListTrack.IsValid())
+		{
+			TUniquePtr<ITimingEventRelation> Relation = MakeUnique<FRenderTraceRelation>(
+				SharedThis(this), CmdList.DetachTime, 0, DetachedCmdListTrack, CmdList.DetachTime, 0, GRelationColorDetach
+			);
+
+			TimingView->AddRelation(Relation);
+		}
+	}
+
+	// If we're detached, add an arrow from the source command list.
+	if (CmdList.Type == ECommandListType::Detached && CmdList.SourceCmdListID != INVALID_EVENT_ID)
+	{
+		const FCommandListInstance& SourceCmdList = Provider->GetCommandList(CmdList.SourceCmdListID);
+		if (SourceCmdList.DetachedCmdListID == CmdListID)
+		{
+			TSharedPtr<const FCommandListTrack> SourceCmdListTrack = SharedData.GetCommandListTrack(SourceCmdList.TimelineIndex);
+			if (SourceCmdListTrack.IsValid())
+			{
+				TUniquePtr<ITimingEventRelation> Relation = MakeUnique<FRenderTraceRelation>(
+					SourceCmdListTrack, SourceCmdList.DetachTime, 0, SharedThis(this), SourceCmdList.DetachTime, 0, GRelationColorDetach
+				);
+
+				TimingView->AddRelation(Relation);
+			}
+		}
+	}
+
+	// Show the event which submitted this command list.
+	if (CmdList.SubmitPassID != INVALID_EVENT_ID)
+	{
+		const FRDGPassInstance& Pass = Provider->GetRDGPass(CmdList.SubmitPassID);
+		TSharedPtr<const FBaseTimingTrack> RTSubmissionTrack = GetFirstChildTrack(SharedData.GetRDGThreadTrack(Pass.ExecThreadID));
+		if (RTSubmissionTrack.IsValid())
+		{
+			TUniquePtr<ITimingEventRelation> Relation = MakeUnique<FRenderTraceRelation>(
+				RTSubmissionTrack, Pass.StartTime, 0, SharedThis(this), CmdList.TimelineStart, 0, GRelationColorSubmitCmdList
+			);
+
+			TimingView->AddRelation(Relation);
+		}
+	}
+
+	// Show the translation task.
+	if (CmdList.TranslateTaskID != INVALID_EVENT_ID)
+	{
+		const FRHITranslateTask& Task = Provider->GetRHITranslateTask(CmdList.TranslateTaskID);
+		TSharedPtr<const FRHIThreadTrack> TaskTrack = SharedData.GetRHIThreadTrack(Task.ThreadID);
+		if (TaskTrack.IsValid())
+		{
+			for (const FEventInterval& CmdListEvent : Task.ProcessedItems)
+			{
+				if (CmdListEvent.ItemID == CmdListID)
+				{
+					TUniquePtr<ITimingEventRelation> Relation = MakeUnique<FRenderTraceRelation>(
+						SharedThis(this), CmdList.FinishRecordingTime, 0, TaskTrack, CmdListEvent.Start, 1, GRelationColorTranslate
+					);
+
+					TimingView->AddRelation(Relation);
+					break;
+				}
+			}
+		}
+	}
+}
+
+void FCommandListTrack::OnRecordingEventSelected(const FTimingEvent& InSelectedEvent) const
+{
+	UE::Insights::Timing::ITimingViewSession* TimingView = SharedData.GetTimingViewSession();
+	const IRenderTraceProvider* Provider = SharedData.GetRenderTraceProvider();
+	check(TimingView && Provider);
+
+	const uint32 PassID = static_cast<uint32>(InSelectedEvent.GetType());
+	const FRDGPassInstance& Pass = Provider->GetRDGPass(PassID);
+	TSharedPtr<const FRDGThreadTrack> PassTrack = SharedData.GetRDGThreadTrack(Pass.ExecThreadID);
+	if (!PassTrack.IsValid())
+	{
+		return;
+	}
+
+	TUniquePtr<ITimingEventRelation> Relation = MakeUnique<FRenderTraceRelation>(
+		PassTrack, Pass.StartTime, 0, SharedThis(this), InSelectedEvent.GetStartTime(), 1, GRelationColorRecording
+	);
+
+	TimingView->AddRelation(Relation);
+}
+
+void FCommandListTrack::OnRenderTraceEventSelected(const FTimingEvent& InSelectedEvent) const
+{
+	UE::Insights::Timing::ITimingViewSession* TimingView = SharedData.GetTimingViewSession();
+	const IRenderTraceProvider* Provider = SharedData.GetRenderTraceProvider();
+	if (!TimingView || !Provider)
+	{
+		return;
+	}
+
+	const uint32 Depth = InSelectedEvent.GetDepth();
+	switch (Depth)
+	{
+	case 0:
+		OnCommandListSelected(InSelectedEvent);
+		break;
+
+	case 1:
+		OnRecordingEventSelected(InSelectedEvent);
+		break;
+
+	default:
+		checkNoEntry();
+		return;
+	}
+}
+
+FString FCommandListTrack::GetEventName(const IRenderTraceProvider* Provider, uint32 InDepth, const FRenderTraceEvent& InEvent) const
+{
+	switch (InDepth)
+	{
+	case 0:
+	{
+		// Command list events.
+		const FCommandListInstance& CmdList = Provider->GetCommandList(InEvent.ItemID);
+		return GetCommandListName(CmdList);
+	}
+
+	case 1:
+	{
+		// Pass recording events.
+		const FRDGPassInstance& Pass = Provider->GetRDGPass(InEvent.ItemID);
+		return Pass.Name;
+	}
+
+	default:
+	{
+		checkNoEntry();
+		return TEXT("invalid");
+	}
+	}
+}
+
+FRDGThreadTrack::FRDGThreadTrack(const FRenderTraceTimingViewSession& InSharedData, const FString& Name, const IRenderTraceProvider::TEventTimeline* InTimeline)
+	: FRenderTraceTrack(InSharedData, Name, InTimeline)
+{
+}
+
+void FRDGThreadTrack::OnRenderTraceEventSelected(const FTimingEvent& InSelectedEvent) const
+{
+	UE::Insights::Timing::ITimingViewSession* TimingView = SharedData.GetTimingViewSession();
+	const IRenderTraceProvider* Provider = SharedData.GetRenderTraceProvider();
+	if (!TimingView || !Provider)
+	{
+		return;
+	}
+
+	const uint32 PassID = static_cast<uint32>(InSelectedEvent.GetType());
+	const FRDGPassInstance& Pass = Provider->GetRDGPass(PassID);
+	check(Pass.Type != ERDGPassType::Submit);
+
+	// Add an arrow to the CPU thread where the pass was executed.
+	TSharedPtr<const FBaseTimingTrack> CpuThreadTrack = SharedData.GetCpuThreadTrack(Pass.ExecThreadID);
+	if (CpuThreadTrack)
+	{
+		TUniquePtr<ITimingEventRelation> Relation = MakeUnique<FRenderTraceRelation>(
+			SharedThis(this), Pass.StartTime, 0, CpuThreadTrack, Pass.StartTime, 0, GRelationColorExecThread
+		);
+		TimingView->AddRelation(Relation);
+	}
+
+	// Add a relationship from the pass to the recording event on each command list it used.
+	for (uint32 CmdListID : Pass.CommandListIDs)
+	{
+		const FCommandListInstance& CmdList = Provider->GetCommandList(CmdListID);
+
+		// Some command lists might not be on the timeline, skip them. TimelineIndex is stored when the analyzer
+		// assigns a command list to a lane.
+		if (CmdList.TimelineIndex < 0)
+		{
+			continue;
+		}
+
+		TSharedPtr<const FCommandListTrack> CmdListTrack = SharedData.GetCommandListTrack(CmdList.TimelineIndex);
+		if (!CmdListTrack.IsValid())
+		{
+			continue;
+		}
+
+		for (const FEventInterval& RecordEvent : CmdList.RecordingEvents)
+		{
+			if (RecordEvent.ItemID != PassID)
+			{
+				continue;
+			}
+
+			// The RDG passes are at depth 0 and the recording events are at depth 1.
+			TUniquePtr<ITimingEventRelation> Relation = MakeUnique<FRenderTraceRelation>(
+				SharedThis(this), Pass.StartTime, 0, CmdListTrack, RecordEvent.Start, 1, GRelationColorRecording
+			);
+
+			TimingView->AddRelation(Relation);
+
+			// There can only be one event per command list.
+			break;
+		}
+	}
+}
+
+void FRDGThreadTrack::PopulateEventDetails(IEventDetailsBuilder& Builder, const FTimingEvent& Event) const
+{
+	const IRenderTraceProvider& Provider = Builder.GetProvider();
+	const uint32 PassID = static_cast<uint32>(Event.GetType());
+	const FRDGPassInstance& Pass = Provider.GetRDGPass(PassID);
+	check(Pass.Type != ERDGPassType::Submit);
+
+	Builder.AddTitle(Pass.Name);
+	Builder.AddNameValueTextLine(TEXTVIEW("Type:"), GetRDGPassTypeName(Pass.Type));
+
+	FString FlagList;
+	bool bShowFlags = false, bShowTaskMode = false, bShowDuration = false;
+
+	switch (Pass.Type)
+	{
+	case ERDGPassType::Regular:
+		bShowFlags = bShowTaskMode = bShowDuration = true;
+		ADD_FLAG_TO_LIST(Pass.Flags, ERDGPassFlag, Raster);
+		ADD_FLAG_TO_LIST(Pass.Flags, ERDGPassFlag, Compute);
+		ADD_FLAG_TO_LIST(Pass.Flags, ERDGPassFlag, AsyncCompute);
+		ADD_FLAG_TO_LIST(Pass.Flags, ERDGPassFlag, Copy);
+		ADD_FLAG_TO_LIST(Pass.Flags, ERDGPassFlag, NeverCull);
+		ADD_FLAG_TO_LIST(Pass.Flags, ERDGPassFlag, SkipRenderPass);
+		ADD_FLAG_TO_LIST(Pass.Flags, ERDGPassFlag, NeverMerge);
+		ADD_FLAG_TO_LIST(Pass.Flags, ERDGPassFlag, NeverParallel);
+		break;
+
+	case ERDGPassType::Dispatch:
+		bShowDuration = true;
+		break;
+
+	default:
+		checkNoEntry();
+		return;
+	}
+
+	if (bShowFlags)
+	{
+		if (FlagList.IsEmpty()) FlagList = TEXT("None");
+		Builder.AddNameValueTextLine(TEXTVIEW("Flags:"), FlagList);
+	}
+
+	if (bShowTaskMode)
+	{
+		const TCHAR* TaskModeStr;
+		switch (Pass.TaskMode)
+		{
+		case ERDGPassTaskMode::Inline: TaskModeStr = TEXT("Inline"); break;
+		case ERDGPassTaskMode::Await: TaskModeStr = TEXT("Await"); break;
+		case ERDGPassTaskMode::Async: TaskModeStr = TEXT("Async"); break;
+		default: TaskModeStr = TEXT("INVALID"); break;
+		}
+
+		Builder.AddNameValueTextLine(TEXTVIEW("Task Mode:"), TaskModeStr);
+	}
+
+	if(bShowDuration)
+	{
+		Builder.AddNameValueTextLine(TEXTVIEW("Duration:"), Insights::FormatTimeAuto(Event.GetDuration()));
+	}
+
+	Builder.AddNameValueTextLine(TEXTVIEW("Command Lists:"), LexToString(Pass.CommandListIDs.Num()));
+	for (uint32 CmdListID : Pass.CommandListIDs)
+	{
+		Builder.AddNameValueTextLine(TEXTVIEW(""), GetCommandListName(Provider.GetCommandList(CmdListID)));
+	}
+}
+
+FString FRDGThreadTrack::GetEventName(const IRenderTraceProvider* Provider, uint32 InDepth, const FRenderTraceEvent& InEvent) const
+{
+	const FRDGPassInstance& Pass = Provider->GetRDGPass(InEvent.ItemID);
+	return Pass.Name;
+}
+
+FRenderThreadSubmissionTrack::FRenderThreadSubmissionTrack(const FRenderTraceTimingViewSession& InSharedData, const FString& Name, const IRenderTraceProvider::TEventTimeline* InTimeline)
+	: FRenderTraceTrack(InSharedData, Name, InTimeline)
+{
+}
+
+void FRenderThreadSubmissionTrack::OnRenderTraceEventSelected(const FTimingEvent& InSelectedEvent) const
+{
+	UE::Insights::Timing::ITimingViewSession* TimingView = SharedData.GetTimingViewSession();
+	const IRenderTraceProvider* Provider = SharedData.GetRenderTraceProvider();
+	if (!TimingView || !Provider)
+	{
+		return;
+	}
+
+	const uint32 PassID = static_cast<uint32>(InSelectedEvent.GetType());
+	const FRDGPassInstance& Pass = Provider->GetRDGPass(PassID);
+	check(Pass.Type == ERDGPassType::Submit);
+
+	// Add an arrow to the CPU thread where the pass was executed.
+	TSharedPtr<const FBaseTimingTrack> CpuThreadTrack = SharedData.GetCpuThreadTrack(Pass.ExecThreadID);
+	if (CpuThreadTrack)
+	{
+		TUniquePtr<ITimingEventRelation> Relation = MakeUnique<FRenderTraceRelation>(
+			SharedThis(this), Pass.StartTime, 0, CpuThreadTrack, Pass.StartTime, 0, GRelationColorExecThread
+		);
+		TimingView->AddRelation(Relation);
+	}
+
+	// Draw arrows from the pass to the command lists it submitted.
+	for (uint32 CmdListID : Pass.CommandListIDs)
+	{
+		const FCommandListInstance& CmdList = Provider->GetCommandList(CmdListID);
+		if (CmdList.TimelineIndex < 0)
+		{
+			continue;
+		}
+
+		TSharedPtr<const FCommandListTrack> CmdListTrack = SharedData.GetCommandListTrack(CmdList.TimelineIndex);
+		if (!CmdListTrack.IsValid())
+		{
+			continue;
+		}
+
+		TUniquePtr<ITimingEventRelation> Relation = MakeUnique<FRenderTraceRelation>(
+			SharedThis(this), Pass.StartTime, 0, CmdListTrack, CmdList.TimelineStart, 0, GRelationColorSubmitCmdList
+		);
+
+		TimingView->AddRelation(Relation);
+	}
+
+	// If we triggered an RHI submit, add an arrow to the RHI submit task.
+	if (Pass.SubmitTaskID != INVALID_EVENT_ID)
+	{
+		const FRHITranslateTask& Task = Provider->GetRHITranslateTask(Pass.SubmitTaskID);
+		check(Task.Type == ERHITranslateTaskType::Submit);
+		TSharedPtr<const FBaseTimingTrack> RHISubmissionTrack = GetFirstChildTrack(SharedData.GetRHIThreadTrack(Task.ThreadID));
+		if (RHISubmissionTrack.IsValid())
+		{
+			TUniquePtr<ITimingEventRelation> Relation = MakeUnique<FRenderTraceRelation>(
+				SharedThis(this), InSelectedEvent.GetEndTime(), 0, RHISubmissionTrack, Task.StartTime, 0, GRelationColorSubmitRHI
+			);
+
+			TimingView->AddRelation(Relation);
+		}
+	}
+}
+
+FString FRenderThreadSubmissionTrack::GetEventName(const IRenderTraceProvider* Provider, uint32 InDepth, const FRenderTraceEvent& InEvent) const
+{
+	const FRDGPassInstance& Pass = Provider->GetRDGPass(InEvent.ItemID);
+	return Pass.Name;
+}
+
+void FRenderThreadSubmissionTrack::PopulateEventDetails(IEventDetailsBuilder& Builder, const FTimingEvent& Event) const
+{
+	const IRenderTraceProvider& Provider = Builder.GetProvider();
+	const uint32 PassID = static_cast<uint32>(Event.GetType());
+	const FRDGPassInstance& Pass = Provider.GetRDGPass(PassID);
+	check(Pass.Type == ERDGPassType::Submit);
+
+	Builder.AddTitle(TEXTVIEW("Submit"));
+	Builder.AddNameValueTextLine(TEXTVIEW("ID:"), LexToString(PassID));
+
+	FString FlagList;
+	ADD_FLAG_TO_LIST(Pass.Flags, ERDGPassSubmitFlag, SubmitToGPU);
+	ADD_FLAG_TO_LIST(Pass.Flags, ERDGPassSubmitFlag, DeleteResources);
+	ADD_FLAG_TO_LIST(Pass.Flags, ERDGPassSubmitFlag, FlushRHIThread);
+	ADD_FLAG_TO_LIST(Pass.Flags, ERDGPassSubmitFlag, EndFrame);
+	ADD_FLAG_TO_LIST(Pass.Flags, ERDGPassSubmitFlag, EnableBypass);
+	ADD_FLAG_TO_LIST(Pass.Flags, ERDGPassSubmitFlag, DisableBypass);
+	ADD_FLAG_TO_LIST(Pass.Flags, ERDGPassSubmitFlag, EnableDrawEvents);
+	ADD_FLAG_TO_LIST(Pass.Flags, ERDGPassSubmitFlag, DisableDrawEvents);
+	if (FlagList.IsEmpty()) FlagList = TEXT("None");
+	Builder.AddNameValueTextLine(TEXTVIEW("Flags:"), FlagList);
+
+	if (Pass.SubmitTaskID != INVALID_EVENT_ID)
+	{
+		Builder.AddNameValueTextLine(TEXTVIEW("RHI Submit Task:"), LexToString(Pass.SubmitTaskID));
+	}
+
+	Builder.AddNameValueTextLine(TEXTVIEW("Command Lists:"), LexToString(Pass.CommandListIDs.Num()));
+	for (uint32 CmdListID : Pass.CommandListIDs)
+	{
+		Builder.AddNameValueTextLine(TEXTVIEW(""), GetCommandListName(Provider.GetCommandList(CmdListID)));
+	}
+}
+
+FRHIThreadTrack::FRHIThreadTrack(const FRenderTraceTimingViewSession& InSharedData, const FString& Name, const IRenderTraceProvider::TEventTimeline* InTimeline)
+	: FRenderTraceTrack(InSharedData, Name, InTimeline)
+{
+}
+
+void FRHIThreadTrack::CacheSplitFromTask(const IRenderTraceProvider& Provider, uint32 TaskID) const
+{
+	// This method will change a field on a task object with only the reader lock held. This is OK to do because it's only ever
+	// called from the UI thread, so there won't be concurrent updates to the cached task ID.
+
+	const FRHITranslateTask& Task = Provider.GetRHITranslateTask(TaskID);
+	if (Task.SplitFromTranslateTaskID < INVALID_EVENT_ID)
+	{
+		// Already cached.
+		return;
+	}
+
+	if (Task.ProcessedItems.IsEmpty())
+	{
+		// No command lists means no split events.
+		Task.SplitFromTranslateTaskID = RELATED_TRANSLATE_JOB_NOT_FOUND;
+		return;
+	}
+
+	// The split information is stored in the first translated command list.
+	const FCommandListInstance& FirstCmdList = Provider.GetCommandList(Task.ProcessedItems[0].ItemID);
+	if (FirstCmdList.TranslatePrevJobAppID == 0)
+	{
+		// No split information.
+		Task.SplitFromTranslateTaskID = RELATED_TRANSLATE_JOB_NOT_FOUND;
+		return;
+	}
+
+	const uint32 PrevTaskID = Provider.FindRHITaskByPredicate(TaskID, 0.05, [&Task, &FirstCmdList](const FRHITranslateTask& OtherTask) -> int
+	{
+		if (OtherTask.AppID == FirstCmdList.TranslatePrevJobAppID) return 1;
+		// If we encounter another task with the same app ID as ourselves it's a reused pointer, so we've moved past our lifetime.
+		return (OtherTask.AppID == Task.AppID) ? -1 : 0;
+	});
+
+	if (PrevTaskID == INVALID_EVENT_ID)
+	{
+		UE_LOG(LogRenderTrace, Warning, TEXT("Translate task %llx (permanent ID %x) split from %llx: could not find previous task."), Task.AppID, TaskID, FirstCmdList.TranslatePrevJobAppID);
+		Task.SplitFromTranslateTaskID = RELATED_TRANSLATE_JOB_NOT_FOUND;
+		return;
+	}
+
+	const FRHITranslateTask& PrevTask = Provider.GetRHITranslateTask(PrevTaskID);
+	if (PrevTask.Type != ERHITranslateTaskType::Translate)
+	{
+		UE_LOG(LogRenderTrace, Warning, TEXT("Translate task %llx (permanent ID %x) split from %llx: found previous task with permanent ID %x, but it's not a translate task."),
+			Task.AppID, TaskID, FirstCmdList.TranslatePrevJobAppID, PrevTaskID);
+		Task.SplitFromTranslateTaskID = RELATED_TRANSLATE_JOB_NOT_FOUND;
+		return;
+	}
+
+	if (PrevTask.SplitToTranslateTaskID != INVALID_EVENT_ID && PrevTask.SplitToTranslateTaskID != TaskID)
+	{
+		UE_LOG(LogRenderTrace, Warning, TEXT("Translate task %llx (permanent ID %x) split from %llx: found previous task with permanent ID %x which already has a different split to ID: %x."),
+			Task.AppID, TaskID, FirstCmdList.TranslatePrevJobAppID, PrevTaskID, PrevTask.SplitToTranslateTaskID);
+		Task.SplitFromTranslateTaskID = RELATED_TRANSLATE_JOB_NOT_FOUND;
+		return;
+	}
+
+	Task.SplitFromTranslateTaskID = PrevTaskID;
+
+	// Link the previous task to this.
+	PrevTask.SplitToTranslateTaskID = TaskID;
+}
+
+void FRHIThreadTrack::CacheSplitToTask(const IRenderTraceProvider& Provider, uint32 TaskID) const
+{
+	// See the comment in CacheSplitFromTask about why it's OK to modify the task object without taking the writer lock.
+
+	const FRHITranslateTask& Task = Provider.GetRHITranslateTask(TaskID);
+	if (Task.SplitToTranslateTaskID < INVALID_EVENT_ID)
+	{
+		// Already cached.
+		return;
+	}
+
+	// Find the closest task which says it split from TaskID.
+	const uint32 NextTaskID = Provider.FindRHITaskByPredicate(TaskID, 0.05,
+		[&Provider, &Task](const FRHITranslateTask& OtherTask) -> int
+		{
+			if (OtherTask.Type == ERHITranslateTaskType::Translate && !OtherTask.ProcessedItems.IsEmpty())
+			{
+				const FCommandListInstance& FirstCmdList = Provider.GetCommandList(OtherTask.ProcessedItems[0].ItemID);
+				if (FirstCmdList.TranslatePrevJobAppID == Task.AppID) return 1;
+			}
+
+			// If we encounter another task with the same app ID as ourselves it's a reused pointer, so we've moved past our lifetime.
+			return (OtherTask.AppID == Task.AppID) ? -1 : 0;
+		}
+	);
+
+	if (NextTaskID == INVALID_EVENT_ID)
+	{
+		Task.SplitToTranslateTaskID = RELATED_TRANSLATE_JOB_NOT_FOUND;
+		return;
+	}
+
+	const FRHITranslateTask& NextTask = Provider.GetRHITranslateTask(NextTaskID);
+	check(NextTask.Type == ERHITranslateTaskType::Translate);
+
+	// Due to pointer reuse, it might be that the next task we found is actually split from a different task which just happens
+	// to have the same pointer value as the task we started with. Do another search starting at the next task to see if we arrive
+	// back at the current task. If not, it's a false match caused by pointer reuse.
+	CacheSplitFromTask(Provider, NextTaskID);
+	if (NextTask.SplitFromTranslateTaskID != TaskID)
+	{
+		Task.SplitToTranslateTaskID = RELATED_TRANSLATE_JOB_NOT_FOUND;
+		return;
+	}
+
+	Task.SplitToTranslateTaskID = NextTaskID;
+}
+
+void FRHIThreadTrack::PopulateTranslateJobDetails(IEventDetailsBuilder& Builder, uint32 TaskID, const FRHITranslateTask& Task) const
+{
+	const IRenderTraceProvider& Provider = Builder.GetProvider();
+
+	Builder.AddTitle(FString::Printf(TEXT("Translate job %s"), *GetTranslateTaskName(Task)));
+
+	{
+		FString FlagList;
+		ADD_FLAG_TO_LIST(Task.JobFlags, ERHITranslateJobFlag, Parallel);
+		ADD_FLAG_TO_LIST(Task.JobFlags, ERHITranslateJobFlag, UsingSubCmdLists);
+		if (FlagList.IsEmpty()) FlagList = TEXT("None");
+		Builder.AddNameValueTextLine(TEXTVIEW("Flags:"), FlagList);
+	}
+
+	if (!Task.ProcessedItems.IsEmpty())
+	{
+		const FCommandListInstance& FirstCmdList = Provider.GetCommandList(Task.ProcessedItems[0].ItemID);
+		FString FlagList;
+		ADD_FLAG_TO_LIST(FirstCmdList.TranslateSplitFlags, ERHITranslateJobSplitFlag, Parallel);
+		ADD_FLAG_TO_LIST(FirstCmdList.TranslateSplitFlags, ERHITranslateJobSplitFlag, Threshold);
+		ADD_FLAG_TO_LIST(FirstCmdList.TranslateSplitFlags, ERHITranslateJobSplitFlag, ParentChild);
+		ADD_FLAG_TO_LIST(FirstCmdList.TranslateSplitFlags, ERHITranslateJobSplitFlag, JumpThreads);
+		if (FlagList.IsEmpty()) FlagList = TEXT("None");
+		Builder.AddNameValueTextLine(TEXTVIEW("Split Reason:"), FlagList);
+
+		if (FirstCmdList.TranslatePrevJobAppID > 0)
+		{
+			Builder.AddNameValueTextLine(TEXTVIEW("Split From:"), FString::Printf(TEXT("0x%llx"), FirstCmdList.TranslatePrevJobAppID));
+		}
+	}
+
+	CacheSplitToTask(Provider, TaskID);
+	if (Task.SplitToTranslateTaskID < RELATED_TRANSLATE_JOB_NOT_FOUND)
+	{
+		const FRHITranslateTask& SplitToTask = Provider.GetRHITranslateTask(Task.SplitToTranslateTaskID);
+		Builder.AddNameValueTextLine(TEXTVIEW("Split To:"), GetTranslateTaskName(SplitToTask));
+	}
+
+	if (Task.NextPhaseTaskID != INVALID_EVENT_ID)
+	{
+		Builder.AddNameValueTextLine(TEXTVIEW("Submitted By:"), LexToString(Task.NextPhaseTaskID));
+	}
+	else
+	{
+		Builder.AddNameValueTextLine(TEXTVIEW("NOT SUBMITTED"), TEXTVIEW(""));
+	}
+
+	FAppIDListBuilder<FCommandListInstance>::AddAllEvents(Builder, TEXT("Command Lists:"), Task.ProcessedItems);
+
+	if (Builder.GetDetailLevel() == EEventDetailLevel::Full)
+	{
+		Builder.AddNameValueTextLine(TEXTVIEW("RDG Passes:"), TEXTVIEW(""));
+		for (const FEventInterval& CmdListEvent : Task.ProcessedItems)
+		{
+			const FCommandListInstance& CmdList = Provider.GetCommandList(CmdListEvent.ItemID);
+			for (const FEventInterval& RecordingEvent : CmdList.RecordingEvents)
+			{
+				const FRDGPassInstance& Pass = Provider.GetRDGPass(RecordingEvent.ItemID);
+				Builder.AddNameValueTextLine(TEXTVIEW(""), Pass.Name);
+			}
+		}
+	}
+}
+void FRHIThreadTrack::ListSyncPoints(IEventDetailsBuilder& Builder, const TCHAR* Title, TConstArrayView<uint32> SyncPointIDs) const
+{
+	if (SyncPointIDs.IsEmpty())
+	{
+		return;
+	}
+
+	const IRenderTraceProvider& Provider = Builder.GetProvider();
+	Builder.AddNameValueTextLine(Title, TEXTVIEW(""));
+	for (uint32 SyncPointID : SyncPointIDs)
+	{
+		const FSyncPoint& SyncPoint = Provider.GetSyncPoint(SyncPointID);
+		Builder.AddNameValueTextLine(TEXTVIEW(""), FString::Printf(TEXT("0x%llx (%s): %llu"), SyncPoint.AppID, GetSyncPointTypeName(SyncPoint.Type), SyncPoint.ResolvedValue));
+	}
+}
+
+void FRHIThreadTrack::PopulatePayloadDetails(IEventDetailsBuilder& Builder, uint32 TranslateJobID, uint32 PipeIdx, uint32 PayloadIdx) const
+{
+	const IRenderTraceProvider& Provider = Builder.GetProvider();
+
+	const FRHITranslateTask& TranslateJob = Provider.GetRHITranslateTask(TranslateJobID);
+	check(TranslateJob.Type == ERHITranslateTaskType::Translate);
+	const FRHITranslateContext& Context = TranslateJob.Contexts[PipeIdx];
+	check(Context.RHIContextID != 0);
+	const FPlatformPayload& Payload = Provider.GetPlatformPayload(Context.PayloadIDs[PayloadIdx]);
+
+	Builder.AddTitle(FString::Printf(TEXT("Payload 0x%llx"), Payload.AppID));
+
+	Builder.AddNameValueTextLine(TEXTVIEW("Pipe:"), GetPipeName(PipeIdx));
+	Builder.AddNameValueTextLine(TEXTVIEW("Context ID:"), FString::Printf(TEXT("0x%llx"), Context.RHIContextID));
+
+	{
+		FAppIDListBuilder<uint64> ListBuilder(Builder, Payload.CmdLists.Num());
+		Builder.AddNameValueTextLine(TEXTVIEW("Platform Command Lists:"), LexToString(Payload.CmdLists.Num()));
+		for (uint64 CmdListID : Payload.CmdLists) { ListBuilder.AddItem(CmdListID); }
+		ListBuilder.Finish();
+	}
+
+	ListSyncPoints(Builder, TEXT("Wait Sync Points:"), Payload.WaitSyncPoints);
+	ListSyncPoints(Builder, TEXT("Signal Sync Points:"), Payload.SignalSyncPoints);
+}
+
+void FRHIThreadTrack::PopulateTranslateContextDetails(IEventDetailsBuilder& Builder, uint32 TranslateJobID, uint32 PipeIdx) const
+{
+	const IRenderTraceProvider& Provider = Builder.GetProvider();
+
+	const FRHITranslateTask& TranslateJob = Provider.GetRHITranslateTask(TranslateJobID);
+	check(TranslateJob.Type == ERHITranslateTaskType::Translate);
+	const FRHITranslateContext& Context = TranslateJob.Contexts[PipeIdx];
+	check(Context.RHIContextID != 0);
+
+	Builder.AddTitle(FString::Printf(TEXT("Translate Context 0x%llx"), Context.RHIContextID));
+	Builder.AddNameValueTextLine(TEXTVIEW("Pipe:"), GetPipeName(PipeIdx));
+
+	uint32 NumCmdLists = 0;
+	
+	{
+		FAppIDListBuilder<FPlatformPayload> ListBuilder(Builder, Context.PayloadIDs.Num());
+		Builder.AddNameValueTextLine(TEXTVIEW("Payloads:"), LexToString(Context.PayloadIDs.Num()));
+		for (uint32 ID : Context.PayloadIDs)
+		{
+			const FPlatformPayload& Payload = Provider.GetPlatformPayload(ID);
+			ListBuilder.AddItem(Payload.AppID);
+			NumCmdLists += Payload.CmdLists.Num();
+		}
+		ListBuilder.Finish();
+	}
+
+	{
+		FAppIDListBuilder<uint64> ListBuilder(Builder, NumCmdLists);
+		Builder.AddNameValueTextLine(TEXTVIEW("Platform Command Lists:"), LexToString(NumCmdLists));
+		for (uint32 ID : Context.PayloadIDs)
+		{
+			const FPlatformPayload& Payload = Provider.GetPlatformPayload(ID);
+			for (uint64 CmdListID: Payload.CmdLists)
+			{
+				ListBuilder.AddItem(CmdListID);
+			}
+		}
+		ListBuilder.Finish();
+	}
+}
+
+void FRHIThreadTrack::PopulateEventDetails(IEventDetailsBuilder& Builder, const FTimingEvent& Event) const
+{
+	const IRenderTraceProvider& Provider = Builder.GetProvider();
+
+	const uint32 EventDepth = Event.GetDepth();
+
+	switch (EventDepth)
+	{
+	case 0:
+	{
+		const uint32 TaskID = static_cast<uint32>(Event.GetType());
+		const FRHITranslateTask& Task = Provider.GetRHITranslateTask(TaskID);
+		check(Task.Type == ERHITranslateTaskType::Translate);
+		PopulateTranslateJobDetails(Builder, TaskID, Task);
+		return;
+	}
+
+	case 1:
+	{
+		const uint32 ItemID = static_cast<uint32>(Event.GetType());
+		const uint32 ParentTaskID = static_cast<uint32>(Event.GetType() >> 32);
+		const FRHITranslateTask& ParentTask = Provider.GetRHITranslateTask(ParentTaskID);
+		check(ParentTask.Type == ERHITranslateTaskType::Translate);
+		const FCommandListInstance& CmdList = Provider.GetCommandList(ItemID);
+		Builder.AddTitle(FString::Printf(TEXT("Command List %s"), *GetCommandListName(CmdList)));
+		return;
+	}
+	}
+
+	const uint32 ItemID = static_cast<uint32>(Event.GetType());
+	const uint32 ParentTaskID = static_cast<uint32>(Event.GetType() >> 32);
+
+	const uint32 FirstContextDepth = 2;
+	const uint32 ContextDepth = EventDepth - FirstContextDepth;
+	const uint32 PipeIdx = ContextDepth / 2;
+	const bool bIsPayload = (ContextDepth % 2) == 1;
+	check(PipeIdx < 2);
+
+	if (bIsPayload)
+	{
+		PopulatePayloadDetails(Builder, ParentTaskID, PipeIdx, ItemID);
+	}
+	else
+	{
+		PopulateTranslateContextDetails(Builder, ParentTaskID, PipeIdx);
+	}
+}
+
+void FRHIThreadTrack::OnTranslateEventSelected(const FTimingEvent& InSelectedEvent) const
+{
+	UE::Insights::Timing::ITimingViewSession* TimingView = SharedData.GetTimingViewSession();
+	const IRenderTraceProvider* Provider = SharedData.GetRenderTraceProvider();
+	check(TimingView && Provider);
+
+	const uint32 TaskID = static_cast<uint32>(InSelectedEvent.GetType());
+	const FRHITranslateTask& Task = Provider->GetRHITranslateTask(TaskID);
+	check(Task.Type == ERHITranslateTaskType::Translate);
+
+	// Add an arrow to the CPU thread where the translation job executed.
+	TSharedPtr<const FBaseTimingTrack> CpuThreadTrack = SharedData.GetCpuThreadTrack(Task.ThreadID);
+	if (CpuThreadTrack)
+	{
+		TUniquePtr<ITimingEventRelation> Relation = MakeUnique<FRenderTraceRelation>(
+			SharedThis(this), Task.StartTime, 0, CpuThreadTrack, Task.StartTime, 0, GRelationColorExecThread
+		);
+		TimingView->AddRelation(Relation);
+	}
+
+	// Draw arrows from all the command lists which are translated by this job.
+	for (int32 CmdListEventIdx = 0; CmdListEventIdx < Task.ProcessedItems.Num(); ++CmdListEventIdx)
+	{
+		const FEventInterval& CmdListEvent = Task.ProcessedItems[CmdListEventIdx];
+		const FCommandListInstance& CmdList = Provider->GetCommandList(CmdListEvent.ItemID);
+
+		if (CmdList.TimelineIndex < 0 || CmdList.FinishRecordingTime == 0.0)
+		{
+			continue;
+		}
+
+		TSharedPtr<const FCommandListTrack> CmdListTrack = SharedData.GetCommandListTrack(CmdList.TimelineIndex);
+		if (!CmdListTrack.IsValid())
+		{
+			continue;
+		}
+
+		TUniquePtr<ITimingEventRelation> Relation = MakeUnique<FRenderTraceRelation>(
+			CmdListTrack, CmdList.FinishRecordingTime, 0, SharedThis(this), CmdListEvent.Start, 1, GRelationColorTranslate
+		);
+
+		TimingView->AddRelation(Relation);
+	}
+
+	// Draw arrow to submit task.
+	if (Task.NextPhaseTaskID != INVALID_EVENT_ID)
+	{
+		const FRHITranslateTask& SubmitTask = Provider->GetRHITranslateTask(Task.NextPhaseTaskID);
+		check(SubmitTask.Type == ERHITranslateTaskType::Submit);
+
+		TSharedPtr<const FBaseTimingTrack> RHISubmissionTrack = GetFirstChildTrack(SharedData.GetRHIThreadTrack(SubmitTask.ThreadID));
+		if (RHISubmissionTrack.IsValid())
+		{
+			for (const FEventInterval& TranslateEvent : SubmitTask.ProcessedItems)
+			{
+				if (TranslateEvent.ItemID == TaskID)
+				{
+					TUniquePtr<ITimingEventRelation> Relation = MakeUnique<FRenderTraceRelation>(
+						SharedThis(this), Task.EndTime, 0, RHISubmissionTrack, TranslateEvent.Start, 1, GRelationColorSubmitRHI
+					);
+
+					TimingView->AddRelation(Relation);
+					break;
+				}
+			}
+		}
+	}
+	
+	CacheSplitFromTask(*Provider, TaskID);
+	if (Task.SplitFromTranslateTaskID < RELATED_TRANSLATE_JOB_NOT_FOUND)
+	{
+		// Draw arrow from translate job which was split into this job.
+		const FRHITranslateTask& PrevTask = Provider->GetRHITranslateTask(Task.SplitFromTranslateTaskID);
+		check(PrevTask.Type == ERHITranslateTaskType::Translate);
+		TSharedPtr<const FRHIThreadTrack> PrevTaskTrack = SharedData.GetRHIThreadTrack(PrevTask.ThreadID);
+		if (PrevTaskTrack.IsValid())
+		{
+			TUniquePtr<ITimingEventRelation> Relation = MakeUnique<FRenderTraceRelation>(
+				PrevTaskTrack, PrevTask.EndTime, 0, SharedThis(this), Task.StartTime, 0, GRelationColorSplitFromTranslateJob
+			);
+
+			TimingView->AddRelation(Relation);
+		}
+	}
+
+	CacheSplitToTask(*Provider, TaskID);
+	if (Task.SplitToTranslateTaskID < RELATED_TRANSLATE_JOB_NOT_FOUND)
+	{
+		// Draw arrow to translate job which followed this task after a split.
+		const FRHITranslateTask& NextTask = Provider->GetRHITranslateTask(Task.SplitToTranslateTaskID);
+		check(NextTask.Type == ERHITranslateTaskType::Translate);
+		TSharedPtr<const FRHIThreadTrack> NextTaskTrack = SharedData.GetRHIThreadTrack(NextTask.ThreadID);
+		if (NextTaskTrack.IsValid())
+		{
+			TUniquePtr<ITimingEventRelation> Relation = MakeUnique<FRenderTraceRelation>(
+				SharedThis(this), Task.EndTime, 0, NextTaskTrack, NextTask.StartTime, 0, GRelationColorSplitToTranslateJob
+			);
+
+			TimingView->AddRelation(Relation);
+		}
+	}
+}
+void FRHIThreadTrack::OnTranslateCommandListEventSelected(const FTimingEvent& InSelectedEvent) const
+{
+	UE::Insights::Timing::ITimingViewSession* TimingView = SharedData.GetTimingViewSession();
+	const IRenderTraceProvider* Provider = SharedData.GetRenderTraceProvider();
+	check(TimingView && Provider);
+
+	// Draw arrow from the command list to the translate event inside the job task.
+	const uint32 CmdListID = static_cast<uint32>(InSelectedEvent.GetType());
+	const FCommandListInstance& CmdList = Provider->GetCommandList(CmdListID);
+	if (CmdList.TimelineIndex < 0 || CmdList.FinishRecordingTime == 0.0)
+	{
+		return;
+	}
+
+	TSharedPtr<const FCommandListTrack> CmdListTrack = SharedData.GetCommandListTrack(CmdList.TimelineIndex);
+	if (!CmdListTrack.IsValid())
+	{
+		return;
+	}
+
+	TUniquePtr<ITimingEventRelation> Relation = MakeUnique<FRenderTraceRelation>(
+		CmdListTrack, CmdList.FinishRecordingTime, 0, SharedThis(this), InSelectedEvent.GetStartTime(), 1, GRelationColorTranslate
+	);
+
+	TimingView->AddRelation(Relation);
+}
+void FRHIThreadTrack::OnPayloadEventSelected(const FTimingEvent& InSelectedEvent, uint32 TranslateJobID, uint8 PipeIdx, uint32 PayloadIdx) const
+{
+	UE::Insights::Timing::ITimingViewSession* TimingView = SharedData.GetTimingViewSession();
+	const IRenderTraceProvider* Provider = SharedData.GetRenderTraceProvider();
+	check(TimingView && Provider);
+
+	const FRHITranslateTask& TranslateTask = Provider->GetRHITranslateTask(TranslateJobID);
+	check(TranslateTask.Type == ERHITranslateTaskType::Translate);
+
+	const uint32 PayloadID = TranslateTask.Contexts[PipeIdx].PayloadIDs[PayloadIdx];
+	const FPlatformPayload& Payload = Provider->GetPlatformPayload(PayloadID);
+
+	TSharedPtr<const FSubmissionQueueTrack> SubmissionTrack = SharedData.GetSubmissionTrack();
+	if (SubmissionTrack.IsValid())
+	{
+		// Draw arrow to the submission event which executed the payload.
+		if (Payload.ExecutionEvent.ParentItemID != INVALID_EVENT_ID)
+		{
+			const FSubmissionBatch& SubmissionBatch = Provider->GetSubmissionBatch(Payload.ExecutionEvent.ParentItemID);
+			const FSubmissionEvent& ExecuteEvent = SubmissionBatch.Events[Payload.ExecutionEvent.ItemID];
+			check(ExecuteEvent.Data.GetIndex() == ESubmissionEvent_Execute);
+
+			TUniquePtr<ITimingEventRelation> Relation = MakeUnique<FRenderTraceRelation>(
+				SharedThis(this), InSelectedEvent.GetEndTime(), InSelectedEvent.GetDepth(), SubmissionTrack, ExecuteEvent.Timestamp, 1, GRelationColorExecute
+			);
+
+			TimingView->AddRelation(Relation);
+		}
+
+		// Draw arrows to all resolve sync point events.
+		for (const FRenderTraceEvent& ResolveEventLocator : Payload.ResolveSyncPointEvents)
+		{
+			const FSubmissionBatch& SubmissionBatch = Provider->GetSubmissionBatch(ResolveEventLocator.ParentItemID);
+			const FSubmissionEvent& ResolveEvent = SubmissionBatch.Events[ResolveEventLocator.ItemID];
+			check(ResolveEvent.Data.GetIndex() == ESubmissionEvent_ResolveSyncPoint);
+
+			TUniquePtr<ITimingEventRelation> Relation = MakeUnique<FRenderTraceRelation>(
+				SharedThis(this), InSelectedEvent.GetEndTime(), InSelectedEvent.GetDepth(), SubmissionTrack, ResolveEvent.Timestamp, 1, GRelationColorResolveSyncPoint
+			);
+
+			TimingView->AddRelation(Relation);
+		}
+	}
+
+	// Draw arrow to the interrupt event which signaled the payload as completed.
+	TSharedPtr<const FInterruptTrack> InterruptTrack = SharedData.GetInterruptTrack();
+	if (InterruptTrack.IsValid() && Payload.InterruptEvent.ParentItemID != INVALID_EVENT_ID)
+	{
+		const FInterruptWakeUp& WakeUp = Provider->GetInterruptWakeUp(Payload.InterruptEvent.ParentItemID);
+		const FInterruptFenceSignaledEvent& SignalEvent = WakeUp.SignalEvents[Payload.InterruptEvent.ItemID];
+
+		TUniquePtr<ITimingEventRelation> Relation = MakeUnique<FRenderTraceRelation>(
+			SharedThis(this), InSelectedEvent.GetEndTime(), InSelectedEvent.GetDepth(), InterruptTrack, SignalEvent.Timestamp, 1, GRelationColorInterrupt
+		);
+
+		TimingView->AddRelation(Relation);
+	}
+}
+
+void FRHIThreadTrack::OnRenderTraceEventSelected(const FTimingEvent& InSelectedEvent) const
+{
+	UE::Insights::Timing::ITimingViewSession* TimingView = SharedData.GetTimingViewSession();
+	const IRenderTraceProvider* Provider = SharedData.GetRenderTraceProvider();
+	if (!Provider || !TimingView)
+	{
+		return;
+	}
+
+	const uint32 Depth = InSelectedEvent.GetDepth();
+	const uint32 ParentItemID = static_cast<uint32>(InSelectedEvent.GetType() >> 32);
+	const uint32 ItemID = static_cast<uint32>(InSelectedEvent.GetType());
+
+	switch (Depth)
+	{
+	case 0:
+	{
+		const FRHITranslateTask& Task = Provider->GetRHITranslateTask(ItemID);
+		check(Task.Type == ERHITranslateTaskType::Translate);
+		OnTranslateEventSelected(InSelectedEvent);
+		return;
+	}
+
+	case 1:
+	{
+		check(ParentItemID != INVALID_EVENT_ID);
+		const FRHITranslateTask& ParentTask = Provider->GetRHITranslateTask(ParentItemID);
+		check(ParentTask.Type == ERHITranslateTaskType::Translate);
+		OnTranslateCommandListEventSelected(InSelectedEvent);
+		return;
+	}
+	}
+
+	const uint32 FirstContextDepth = 2;
+	const uint32 ContextDepth = Depth - FirstContextDepth;
+	const uint32 PipeIdx = ContextDepth / 2;
+	if (PipeIdx > 1)
+	{
+		return;
+	}
+
+	const bool bIsPayload = (ContextDepth % 2) == 1;
+	if (bIsPayload)
+	{
+		OnPayloadEventSelected(InSelectedEvent, ParentItemID, PipeIdx, ItemID);
+	}
+}
+
+FString FRHIThreadTrack::GetEventName(const IRenderTraceProvider* Provider, uint32 InDepth, const FRenderTraceEvent& InEvent) const
+{
+	switch (InDepth)
+	{
+	case 0:
+	{
+		const FRHITranslateTask& Task = Provider->GetRHITranslateTask(InEvent.ItemID);
+		check(Task.Type == ERHITranslateTaskType::Translate);
+		return FString::Printf(TEXT("Translate %s"), *GetTranslateTaskName(Task));
+	}
+
+	case 1:
+	{
+		check(InEvent.ParentItemID != INVALID_EVENT_ID);
+		const FRHITranslateTask& ParentTask = Provider->GetRHITranslateTask(InEvent.ParentItemID);
+		check(ParentTask.Type == ERHITranslateTaskType::Translate);
+		const FCommandListInstance& CmdList = Provider->GetCommandList(InEvent.ItemID);
+		return FString::Printf(TEXT("CMD %s"), *GetCommandListName(CmdList));
+	}
+	}
+
+	checkNoEntry();
+	return TEXT("invalid");
+}
+
+void FRHIThreadTrack::CustomizeEventDrawState(ITimingEventsTrackDrawStateBuilder& Builder, double StartTime, double EndTime, uint32 Depth, const FRenderTraceEvent& Event) const
+{
+	// For each translate job event we will add context and payload events starting at depth 2.
+	if (Depth > 0)
+	{
+		return;
+	}
+
+	const IRenderTraceProvider* Provider = SharedData.GetRenderTraceProvider();
+	if (!Provider)
+	{
+		return;
+	}
+
+	const FRHITranslateTask& Task = Provider->GetRHITranslateTask(Event.ItemID);
+	if (Task.Type != ERHITranslateTaskType::Translate || Task.ProcessedItems.IsEmpty())
+	{
+		return;
+	}
+
+	for (int PipeIdx = 0; PipeIdx < 2; ++PipeIdx)
+	{
+		const int ContextDepth = 2 + 2 * PipeIdx;
+		const FRHITranslateContext& Context = Task.Contexts[PipeIdx];
+		const FString ContextName = FString::Printf(TEXT("CTX 0x%llx"), Context.RHIContextID);
+
+		for (const FTimeInterval& Interval : Context.ActiveIntervals)
+		{
+			const double CtxEndTime = Interval.End != 0 ? FMath::Min(Interval.End, EndTime) : EndTime;
+			Builder.AddEvent(Interval.Start, CtxEndTime, ContextDepth, *ContextName);
+		}
+
+		for (uint32 PayloadID : Context.PayloadIDs)
+		{
+			const FPlatformPayload& Payload = Provider->GetPlatformPayload(PayloadID);
+			FString PayloadName = FString::Printf(TEXT("Payload 0x%llx"), Payload.AppID);
+			Builder.AddEvent(Payload.StartTime, Payload.EndTime, ContextDepth + 1, *PayloadName);
+		}
+	}
+}
+
+TSharedPtr<FTimingEvent> FRHIThreadTrack::GetEventAtTime(double EventTime, int32 EventDepth, double TimelineEventStart, double TimelineEventEnd, int32 TimelineEventDepth, const FRenderTraceEvent& TimelineEvent) const
+{
+	const IRenderTraceProvider* Provider = SharedData.GetRenderTraceProvider();
+	if (!Provider)
+	{
+		return nullptr;
+	}
+
+	// We customize the context and platform context events, which start under the job and command list lanes.
+	const int32 FirstContextDepth = 2;
+	if (EventDepth < FirstContextDepth)
+	{
+		return nullptr;
+	}
+
+	// The information for these events is based on the translate job, which is at depth 0 on the timeline.
+	if (TimelineEventDepth > 0)
+	{
+		return nullptr;
+	}
+
+	// We have two lanes per pipe: context and platform command lists.
+	const int32 ContextDepth = EventDepth - FirstContextDepth;
+	const int32 PipeIdx = ContextDepth / 2;
+	if (PipeIdx >= 2)
+	{
+		return nullptr;
+	}
+
+	const bool bIsPayload = (ContextDepth % 2) == 1;
+
+	const FRHITranslateTask& TranslateJob = Provider->GetRHITranslateTask(TimelineEvent.ItemID);
+	if (TranslateJob.Type != ERHITranslateTaskType::Translate)
+	{
+		return nullptr;
+	}
+
+	const FRHITranslateContext& Context = TranslateJob.Contexts[PipeIdx];
+	if (Context.ActiveIntervals.IsEmpty())
+	{
+		return nullptr;
+	}
+
+	// For payloads, the item ID stores the payload index. For contexts we don't need to store anything, since the
+	// depth tells us if it's a context of a payload.
+	uint32 ItemID = INVALID_EVENT_ID;
+
+	double EventStart = TimelineEventStart, EventEnd = TimelineEventEnd;
+
+	if (bIsPayload)
+	{
+		// We must find the payload which was active at EventTime. We usually have a small number of payloads,
+		// so a linear search is good enough.
+		for (int PayloadIdx = 0; PayloadIdx < Context.PayloadIDs.Num(); ++PayloadIdx)
+		{
+			const FPlatformPayload& Payload = Provider->GetPlatformPayload(Context.PayloadIDs[PayloadIdx]);
+			if (EventTime >= Payload.StartTime && EventTime <= Payload.EndTime)
+			{
+				ItemID = PayloadIdx;
+				EventStart = Payload.StartTime;
+				EventEnd = Payload.EndTime;
+				break;
+			}
+		}
+
+		if (ItemID == INVALID_EVENT_ID)
+		{
+			return nullptr;
+		}
+	}
+	else
+	{
+		// We create a single event for the entire lifetime of the context instead of individual events for each
+		// activity interval, to highlight that it's in fact a single context for the entire lane.
+		EventStart = Context.ActiveIntervals[0].Start;
+		EventEnd = Context.ActiveIntervals.Last().End;
+	}
+
+	EventEnd = (EventEnd != 0.0) ? FMath::Min(EventEnd, TranslateJob.EndTime) : TranslateJob.EndTime;
+
+	// Use the translate job ID as the parent ID.
+	const uint64 Packed = ((uint64)TimelineEvent.ItemID << 32) | ItemID;
+	return MakeShared<FTimingEvent>(SharedThis(this), EventStart, EventEnd, EventDepth, Packed);
+}
+
+FRHISubmissionTrack::FRHISubmissionTrack(const FRenderTraceTimingViewSession& InSharedData, const FString& Name, const IRenderTraceProvider::TEventTimeline* InTimeline)
+	: FRenderTraceTrack(InSharedData, Name, InTimeline)
+{
+}
+
+void FRHISubmissionTrack::OnSubmitEventSelected(const FTimingEvent& InSelectedEvent) const
+{
+	UE::Insights::Timing::ITimingViewSession* TimingView = SharedData.GetTimingViewSession();
+	const IRenderTraceProvider* Provider = SharedData.GetRenderTraceProvider();
+	check(TimingView && Provider);
+
+	const uint32 SubmitTaskID = static_cast<uint32>(InSelectedEvent.GetType());
+	const FRHITranslateTask& SubmitTask = Provider->GetRHITranslateTask(SubmitTaskID);
+	check(SubmitTask.Type == ERHITranslateTaskType::Submit);
+
+	// Draw arrow from instigating render thread task.
+	if (SubmitTask.SubmitPassID != INVALID_EVENT_ID)
+	{
+		const FRDGPassInstance& SubmitPass = Provider->GetRDGPass(SubmitTask.SubmitPassID);
+		TSharedPtr<const FBaseTimingTrack> RTSubmissionTrack = GetFirstChildTrack(SharedData.GetRDGThreadTrack(SubmitPass.ExecThreadID));
+		if (RTSubmissionTrack.IsValid())
+		{
+			TUniquePtr<ITimingEventRelation> Relation = MakeUnique<FRenderTraceRelation>(
+				RTSubmissionTrack, SubmitPass.EndTime, 0, SharedThis(this), SubmitTask.StartTime, 0, GRelationColorSubmitRHI
+			);
+
+			TimingView->AddRelation(Relation);
+		}
+	}
+
+	// Draw arrows from all translate jobs which are submitted.
+	for (const FEventInterval& TranslateEvent : SubmitTask.ProcessedItems)
+	{
+		const FRHITranslateTask& TranslateTask = Provider->GetRHITranslateTask(TranslateEvent.ItemID);
+		check(TranslateTask.Type == ERHITranslateTaskType::Translate);
+		TSharedPtr<const FRHIThreadTrack> TranslateTrack = SharedData.GetRHIThreadTrack(TranslateTask.ThreadID);
+		if (TranslateTrack.IsValid())
+		{
+			TUniquePtr<ITimingEventRelation> Relation = MakeUnique<FRenderTraceRelation>(
+				TranslateTrack, TranslateTask.EndTime, 0, SharedThis(this), TranslateEvent.Start, 1, GRelationColorSubmitRHI
+			);
+
+			TimingView->AddRelation(Relation);
+		}
+	}
+}
+
+void FRHISubmissionTrack::OnSubmitTranslateJobEventSelected(const FTimingEvent& InSelectedEvent) const
+{
+	UE::Insights::Timing::ITimingViewSession* TimingView = SharedData.GetTimingViewSession();
+	const IRenderTraceProvider* Provider = SharedData.GetRenderTraceProvider();
+	check(TimingView && Provider);
+
+	// Draw arrow from the translate job to the event inside the submit task.
+	const uint32 TranslateTaskID = static_cast<uint32>(InSelectedEvent.GetType());
+	const FRHITranslateTask& TranslateTask = Provider->GetRHITranslateTask(TranslateTaskID);
+	check(TranslateTask.Type == ERHITranslateTaskType::Translate);
+
+	TSharedPtr<const FRHIThreadTrack> TranslateTrack = SharedData.GetRHIThreadTrack(TranslateTask.ThreadID);
+	if (!TranslateTrack.IsValid())
+	{
+		return;
+	}
+
+	TUniquePtr<ITimingEventRelation> Relation = MakeUnique<FRenderTraceRelation>(
+		TranslateTrack, TranslateTask.EndTime, 0, SharedThis(this), InSelectedEvent.GetStartTime(), 1, GRelationColorSubmitRHI
+	);
+
+	TimingView->AddRelation(Relation);
+}
+
+void FRHISubmissionTrack::OnRenderTraceEventSelected(const FTimingEvent& InSelectedEvent) const
+{
+	const IRenderTraceProvider* Provider = SharedData.GetRenderTraceProvider();
+	UE::Insights::Timing::ITimingViewSession* TimingView = SharedData.GetTimingViewSession();
+	if (!Provider || !TimingView)
+	{
+		return;
+	}
+
+	const uint32 Depth = InSelectedEvent.GetDepth();
+
+	switch (Depth)
+	{
+	case 0:
+	{
+		const uint32 TaskID = static_cast<uint32>(InSelectedEvent.GetType());
+		const FRHITranslateTask& Task = Provider->GetRHITranslateTask(TaskID);
+		check(Task.Type == ERHITranslateTaskType::Submit);
+		OnSubmitEventSelected(InSelectedEvent);
+		return;
+	}
+
+	case 1:
+	{
+		const uint32 ParentItemID = static_cast<uint32>(InSelectedEvent.GetType() >> 32);
+		check(ParentItemID != INVALID_EVENT_ID);
+		const FRHITranslateTask& ParentTask = Provider->GetRHITranslateTask(ParentItemID);
+		check(ParentTask.Type == ERHITranslateTaskType::Submit);
+		OnSubmitTranslateJobEventSelected(InSelectedEvent);
+		return;
+	}
+	}
+
+	checkNoEntry();
+}
+
+FString FRHISubmissionTrack::GetEventName(const IRenderTraceProvider* Provider, uint32 InDepth, const FRenderTraceEvent& InEvent) const
+{
+	switch (InDepth)
+	{
+	case 0:
+	{
+		const FRHITranslateTask& Task = Provider->GetRHITranslateTask(InEvent.ItemID);
+		check(Task.Type == ERHITranslateTaskType::Submit);
+		return FString::Printf(TEXT("Submit %s"), *GetTranslateTaskName(Task));
+	}
+
+	case 1:
+	{
+		check(InEvent.ParentItemID != INVALID_EVENT_ID);
+		const FRHITranslateTask& TranslateTask = Provider->GetRHITranslateTask(InEvent.ItemID);
+		return GetTranslateTaskName(TranslateTask);
+	}
+	}
+
+	checkNoEntry();
+	return TEXT("invalid");
+}
+
+void FRHISubmissionTrack::PopulateSubmitTaskDetails(IEventDetailsBuilder& Builder, uint32 TaskID, const FRHITranslateTask& Task) const
+{
+	const IRenderTraceProvider& Provider = Builder.GetProvider();
+
+	Builder.AddTitle(TEXTVIEW("Submit Task"));
+	Builder.AddNameValueTextLine(TEXTVIEW("ID:"), LexToString(Task.AppID));
+
+	if (Task.SubmitPassID != INVALID_EVENT_ID)
+	{
+		Builder.AddNameValueTextLine(TEXTVIEW("Triggered By:"), LexToString(Task.SubmitPassID));
+	}
+	else
+	{
+		Builder.AddNameValueTextLine(TEXTVIEW("NOT TRIGGERED"), TEXTVIEW(""));
+	}
+
+	FAppIDListBuilder<FRHITranslateTask>::AddAllEvents(Builder, TEXT("Translate Jobs:"), Task.ProcessedItems);
+	
+	uint32 NumCmdLists = 0, NumPasses = 0;
+	for (const FEventInterval& TranslateTaskEvent : Task.ProcessedItems)
+	{
+		const FRHITranslateTask& TranslateTask = Provider.GetRHITranslateTask(TranslateTaskEvent.ItemID);
+		NumCmdLists += TranslateTask.ProcessedItems.Num();
+		for (const FEventInterval& CmdListEvent : TranslateTask.ProcessedItems)
+		{
+			const FCommandListInstance& CmdList = Provider.GetCommandList(CmdListEvent.ItemID);
+			NumPasses += CmdList.RecordingEvents.Num();
+		}
+	}
+
+	Builder.AddNameValueTextLine(TEXTVIEW("Command Lists:"), LexToString(NumCmdLists));
+	FAppIDListBuilder<FCommandListInstance> CmdListNames(Builder, NumCmdLists);
+	for (const FEventInterval& TranslateTaskEvent : Task.ProcessedItems)
+	{
+		const FRHITranslateTask& TranslateTask = Provider.GetRHITranslateTask(TranslateTaskEvent.ItemID);
+		CmdListNames.AddEvents(TranslateTask.ProcessedItems, false);
+	}
+	CmdListNames.Finish();
+
+	if (Builder.GetDetailLevel() == EEventDetailLevel::Full)
+	{
+		Builder.AddNameValueTextLine(TEXTVIEW("RDG Passes:"), LexToString(NumPasses));
+		for (const FEventInterval& TranslateTaskEvent : Task.ProcessedItems)
+		{
+			const FRHITranslateTask& TranslateTask = Provider.GetRHITranslateTask(TranslateTaskEvent.ItemID);
+			for (const FEventInterval& CmdListEvent : TranslateTask.ProcessedItems)
+			{
+				const FCommandListInstance& CmdList = Provider.GetCommandList(CmdListEvent.ItemID);
+				for (const FEventInterval& RecordingEvent : CmdList.RecordingEvents)
+				{
+					const FRDGPassInstance& Pass = Provider.GetRDGPass(RecordingEvent.ItemID);
+					Builder.AddNameValueTextLine(TEXTVIEW(""), Pass.Name);
+				}
+			}
+		}
+	}
+}
+
+void FRHISubmissionTrack::PopulateEventDetails(IEventDetailsBuilder& Builder, const FTimingEvent& Event) const
+{
+	const IRenderTraceProvider& Provider = Builder.GetProvider();
+
+	const uint32 EventDepth = Event.GetDepth();
+
+	switch (EventDepth)
+	{
+	case 0:
+	{
+		const uint32 TaskID = static_cast<uint32>(Event.GetType());
+		const FRHITranslateTask& Task = Provider.GetRHITranslateTask(TaskID);
+		check(Task.Type == ERHITranslateTaskType::Submit);
+		PopulateSubmitTaskDetails(Builder, TaskID, Task);
+		return;
+	}
+
+	case 1:
+	{
+		const uint32 ItemID = static_cast<uint32>(Event.GetType());
+		const uint32 ParentTaskID = static_cast<uint32>(Event.GetType() >> 32);
+		const FRHITranslateTask& ParentTask = Provider.GetRHITranslateTask(ParentTaskID);
+		check(ParentTask.Type == ERHITranslateTaskType::Submit);
+		const FRHITranslateTask& Task = Provider.GetRHITranslateTask(ItemID);
+		Builder.AddTitle(FString::Printf(TEXT("Translate job %s"), *GetTranslateTaskName(Task)));
+		return;
+	}
+	}
+
+	checkNoEntry();
+}
+
+FSubmissionQueueTrack::FSubmissionQueueTrack(const FRenderTraceTimingViewSession& InSharedData, const FString& Name, const IRenderTraceProvider::TEventTimeline* InTimeline)
+	: FRenderTraceTrack(InSharedData, Name, InTimeline)
+{
+}
+
+void FSubmissionQueueTrack::OnSubmissionBatchSelected(const FTimingEvent& InSelectedEvent, uint32 BatchID) const
+{
+	UE::Insights::Timing::ITimingViewSession* TimingView = SharedData.GetTimingViewSession();
+	const IRenderTraceProvider* Provider = SharedData.GetRenderTraceProvider();
+	check(TimingView && Provider);
+
+	const FSubmissionBatch& Batch = Provider->GetSubmissionBatch(BatchID);
+
+	// Add an arrow to the CPU thread where batch was processed.
+	TSharedPtr<const FBaseTimingTrack> CpuThreadTrack = SharedData.GetCpuThreadTrack(Batch.ThreadID);
+	if (CpuThreadTrack)
+	{
+		TUniquePtr<ITimingEventRelation> Relation = MakeUnique<FRenderTraceRelation>(
+			SharedThis(this), Batch.StartTime, 0, CpuThreadTrack, Batch.StartTime, 0, GRelationColorExecThread
+		);
+		TimingView->AddRelation(Relation);
+	}
+}
+
+void FSubmissionQueueTrack::OnResolveSyncPointEventSelected(const FTimingEvent& InSelectedEvent, const FSubmissionEvent::FResolveSyncPoint& ResolveEvent) const
+{
+	UE::Insights::Timing::ITimingViewSession* TimingView = SharedData.GetTimingViewSession();
+	const IRenderTraceProvider* Provider = SharedData.GetRenderTraceProvider();
+	check(TimingView && Provider);
+
+	const FSyncPoint& SyncPoint = Provider->GetSyncPoint(ResolveEvent.SyncPointID);
+	if(SyncPoint.ResolvedByPayload == INVALID_EVENT_ID)
+	{
+		return;
+	}
+	
+	// Draw arrow from the payload which resolved the sync point.
+	const FPlatformPayload& Payload = Provider->GetPlatformPayload(SyncPoint.ResolvedByPayload);
+	const FRHITranslateTask& TranslateTask = Provider->GetRHITranslateTask(Payload.TranslateEvent.ParentItemID);
+	check(TranslateTask.Type == ERHITranslateTaskType::Translate);
+
+	TSharedPtr<const FRHIThreadTrack> TranslateTrack = SharedData.GetRHIThreadTrack(TranslateTask.ThreadID);
+	if (!TranslateTrack.IsValid())
+	{
+		return;
+	}
+
+	const uint32 FirstContextDepth = 2;
+	const uint32 SourceDepth = FirstContextDepth + 2 * Payload.PipeIdx + 1;
+
+	TUniquePtr<ITimingEventRelation> Relation = MakeUnique<FRenderTraceRelation>(
+		TranslateTrack, Payload.EndTime, SourceDepth, SharedThis(this), InSelectedEvent.GetStartTime(), 1, GRelationColorResolveSyncPoint
+	);
+
+	TimingView->AddRelation(Relation);
+}
+
+void FSubmissionQueueTrack::OnExecuteEventSelected(const FTimingEvent& InSelectedEvent, const FSubmissionEvent::FExecute& ExecEvent) const
+{
+	UE::Insights::Timing::ITimingViewSession* TimingView = SharedData.GetTimingViewSession();
+	const IRenderTraceProvider* Provider = SharedData.GetRenderTraceProvider();
+	check(TimingView && Provider);
+
+	for (uint32 PayloadID : ExecEvent.PayloadIDs)
+	{
+		const FPlatformPayload& Payload = Provider->GetPlatformPayload(PayloadID);
+		if (Payload.TranslateEvent.ParentItemID == INVALID_EVENT_ID)
+		{
+			continue;
+		}
+
+		const FRHITranslateTask& TranslateTask = Provider->GetRHITranslateTask(Payload.TranslateEvent.ParentItemID);
+		check(TranslateTask.Type == ERHITranslateTaskType::Translate);
+
+		// Draw arrow from the translate event which produced the payload.
+		TSharedPtr<const FRHIThreadTrack> TranslateTrack = SharedData.GetRHIThreadTrack(TranslateTask.ThreadID);
+		if (!TranslateTrack.IsValid())
+		{
+			continue;
+		}
+
+		const uint32 FirstContextDepth = 2;
+		const uint32 SourceDepth = FirstContextDepth + 2 * Payload.PipeIdx + 1;
+
+		TUniquePtr<ITimingEventRelation> Relation = MakeUnique<FRenderTraceRelation>(
+			TranslateTrack, Payload.EndTime, SourceDepth, SharedThis(this), InSelectedEvent.GetStartTime(), 1, GRelationColorExecute
+		);
+
+		TimingView->AddRelation(Relation);
+	}
+}
+
+void FSubmissionQueueTrack::OnRenderTraceEventSelected(const FTimingEvent& InSelectedEvent) const
+{
+	UE::Insights::Timing::ITimingViewSession* TimingView = SharedData.GetTimingViewSession();
+	const IRenderTraceProvider* Provider = SharedData.GetRenderTraceProvider();
+	if (!Provider || !TimingView)
+	{
+		return;
+	}
+
+	const uint32 Depth = InSelectedEvent.GetDepth();
+	const uint32 ParentItemID = static_cast<uint32>(InSelectedEvent.GetType() >> 32);
+	const uint32 ItemID = static_cast<uint32>(InSelectedEvent.GetType());
+
+	if (Depth == 0)
+	{
+		OnSubmissionBatchSelected(InSelectedEvent, ItemID);
+		return;
+	}
+
+	check(Depth == 1);
+
+	const FSubmissionBatch& Batch = Provider->GetSubmissionBatch(ParentItemID);
+	const FSubmissionEvent& Event = Batch.Events[ItemID];
+
+	switch (Event.Data.GetIndex())
+	{
+	case ESubmissionEvent_ResolveSyncPoint: OnResolveSyncPointEventSelected(InSelectedEvent, Event.Data.Get<FSubmissionEvent::FResolveSyncPoint>()); return;
+	case ESubmissionEvent_Execute: OnExecuteEventSelected(InSelectedEvent, Event.Data.Get<FSubmissionEvent::FExecute>()); return;
+	}
+}
+
+static FString GetSubmissionEventTypeName(const FSubmissionEvent& SubmissionEvent)
+{
+	const TCHAR* TypeName;
+	switch (SubmissionEvent.Data.GetIndex())
+	{
+	case ESubmissionEvent_YieldSyncPoint:     TypeName = TEXT("Yield SyncPoint"); break;
+	case ESubmissionEvent_YieldManualFence:   TypeName = TEXT("Yield Manual Fence"); break;
+	case ESubmissionEvent_WaitQueueFence:     TypeName = TEXT("Wait Queue Fence"); break;
+	case ESubmissionEvent_WaitManualFence:    TypeName = TEXT("Wait Manual Fence"); break;
+	case ESubmissionEvent_Execute:            TypeName = TEXT("Execute"); break;
+	case ESubmissionEvent_SignalManualFence:  TypeName = TEXT("Signal Manual Fence"); break;
+	case ESubmissionEvent_SignalQueueFence:   TypeName = TEXT("Signal Queue Fence"); break;
+	case ESubmissionEvent_ResolveSyncPoint:   TypeName = TEXT("Resolve Sync Point"); break;
+	default:                                  TypeName = TEXT("UNKNOWN"); break;
+	}
+	return FString::Printf(TEXT("[%s] %s"), GetPipeName(SubmissionEvent.Queue), TypeName);
+}
+
+FString FSubmissionQueueTrack::GetEventName(const IRenderTraceProvider* Provider, uint32 InDepth, const FRenderTraceEvent& InEvent) const
+{
+	if (InDepth == 0)
+	{
+		return TEXT("ProcessSubmissionQueue");
+	}
+
+	const FSubmissionBatch& Batch = Provider->GetSubmissionBatch(InEvent.ParentItemID);
+	const FSubmissionEvent& SubmissionEvent = Batch.Events[InEvent.ItemID];
+	return GetSubmissionEventTypeName(SubmissionEvent);
+}
+
+void FSubmissionQueueTrack::PopulateSubmissionBatchDetails(IEventDetailsBuilder& Builder, uint32 ItemID) const
+{
+	const IRenderTraceProvider& Provider = Builder.GetProvider();
+	const FSubmissionBatch& Batch = Provider.GetSubmissionBatch(ItemID);
+
+	Builder.AddTitle(TEXTVIEW("ProcessSubmissionQueue"));
+	Builder.AddNameValueTextLine(TEXTVIEW("Exit Status:"), GetExitStatusStr(Batch.ExitStatus));
+
+	Builder.AddNameValueTextLine(TEXTVIEW("Events:"), LexToString(Batch.Events.Num()));
+
+	uint32 NumPayloads = 0, NumCmdLists = 0;
+	for (const FSubmissionEvent& Event : Batch.Events)
+	{
+		if (Event.Data.IsType<FSubmissionEvent::FExecute>())
+		{
+			const FSubmissionEvent::FExecute& ECL = Event.Data.Get<FSubmissionEvent::FExecute>();
+			NumPayloads += ECL.PayloadIDs.Num();
+			NumCmdLists += ECL.PlatformCmdListIDs.Num();
+		}
+	}
+
+	Builder.AddNameValueTextLine(TEXTVIEW("Executed Payloads:"), LexToString(NumPayloads));
+	Builder.AddNameValueTextLine(TEXTVIEW("Executed CmdLists:"), LexToString(NumCmdLists));
+
+	if (Builder.GetDetailLevel() == EEventDetailLevel::Brief)
+	{
+		return;
+	}
+
+	for (const FSubmissionEvent& Event : Batch.Events)
+	{
+		const FString EventLabel = GetSubmissionEventTypeName(Event) + TEXT(":");
+
+		switch (Event.Data.GetIndex())
+		{
+		case ESubmissionEvent_YieldSyncPoint:
+		{
+			const FSubmissionEvent::FYieldSyncPoint& YieldEvent = Event.Data.Get<FSubmissionEvent::FYieldSyncPoint>();
+			Builder.AddNameValueTextLine(EventLabel, FString::Printf(TEXT("payload=0x%llx syncpoint=0x%llx"), YieldEvent.PayloadID, YieldEvent.SyncPointID));
+			break;
+		}
+		case ESubmissionEvent_YieldManualFence:
+		{
+			const FSubmissionEvent::FYieldManualFence& YieldEvent = Event.Data.Get<FSubmissionEvent::FYieldManualFence>();
+			Builder.AddNameValueTextLine(EventLabel, FString::Printf(TEXT("payload=0x%llx fence=0x%llx value=%llu"), YieldEvent.PayloadID, YieldEvent.FenceID, YieldEvent.FenceValue));
+			break;
+		}
+		case ESubmissionEvent_WaitQueueFence:
+		{
+			const FSubmissionEvent::FWaitQueueFence& WaitQueueFenceEvent = Event.Data.Get<FSubmissionEvent::FWaitQueueFence>();
+			Builder.AddNameValueTextLine(EventLabel, FString::Printf(TEXT("wait on %s value=%llu payload=0x%llx"), GetPipeName(WaitQueueFenceEvent.OtherQueue), WaitQueueFenceEvent.Value, WaitQueueFenceEvent.PayloadID));
+			break;
+		}
+		case ESubmissionEvent_WaitManualFence:
+		{
+			const FSubmissionEvent::FWaitManualFence& WaitManualFenceEvent = Event.Data.Get<FSubmissionEvent::FWaitManualFence>();
+			Builder.AddNameValueTextLine(EventLabel, FString::Printf(TEXT("fence=0x%llx value=%llu payload=0x%llx"), WaitManualFenceEvent.FenceID, WaitManualFenceEvent.Value, WaitManualFenceEvent.PayloadID));
+			break;
+		}
+		case ESubmissionEvent_Execute:
+		{
+			const FSubmissionEvent::FExecute& ExecuteEvent = Event.Data.Get<FSubmissionEvent::FExecute>();
+			Builder.AddNameValueTextLine(EventLabel, TEXTVIEW(""));
+			{
+				Builder.AddNameValueTextLine(TEXTVIEW("Payloads:"), LexToString(ExecuteEvent.PayloadIDs.Num()));
+				FAppIDListBuilder<uint64> ListBuilder(Builder, ExecuteEvent.PayloadIDs.Num());
+				for (uint32 PayloadID : ExecuteEvent.PayloadIDs) { ListBuilder.AddItem(Provider.GetPlatformPayload(PayloadID).AppID); }
+				ListBuilder.Finish();
+			}
+			{
+				Builder.AddNameValueTextLine(TEXTVIEW("Platform CmdLists:"), LexToString(ExecuteEvent.PlatformCmdListIDs.Num()));
+				FAppIDListBuilder<uint64> ListBuilder(Builder, ExecuteEvent.PlatformCmdListIDs.Num());
+				ListBuilder.AddIDs(ExecuteEvent.PlatformCmdListIDs, true);
+			}
+			break;
+		}
+		case ESubmissionEvent_SignalManualFence:
+		{
+			const FSubmissionEvent::FSignalManualFence& SignalManualFenceEvent = Event.Data.Get<FSubmissionEvent::FSignalManualFence>();
+			Builder.AddNameValueTextLine(EventLabel, FString::Printf(TEXT("fence=0x%llx value=%llu payload=0x%llx"), SignalManualFenceEvent.FenceID, SignalManualFenceEvent.Value, SignalManualFenceEvent.PayloadID));
+			break;
+		}
+		case ESubmissionEvent_SignalQueueFence:
+		{
+			const FSubmissionEvent::FSignalQueueFence& SignalQueueFenceEvent = Event.Data.Get<FSubmissionEvent::FSignalQueueFence>();
+			Builder.AddNameValueTextLine(EventLabel, FString::Printf(TEXT("value=%llu payload=0x%llx"), SignalQueueFenceEvent.Value, SignalQueueFenceEvent.PayloadID));
+			break;
+		}
+		case ESubmissionEvent_ResolveSyncPoint:
+		{
+			const FSubmissionEvent::FResolveSyncPoint& ResolveSyncPointEvent = Event.Data.Get<FSubmissionEvent::FResolveSyncPoint>();
+			Builder.AddNameValueTextLine(EventLabel, FString::Printf(TEXT("payload=0x%llx syncpoint=0x%llx value=%llu"), Provider.GetPlatformPayload(ResolveSyncPointEvent.PayloadID).AppID, Provider.GetSyncPoint(ResolveSyncPointEvent.SyncPointID).AppID, ResolveSyncPointEvent.Value));
+			break;
+		}
+		}
+	}
+}
+
+void FSubmissionQueueTrack::PopulateEventDetails(IEventDetailsBuilder& Builder, const FTimingEvent& Event) const
+{
+	const uint32 Depth = Event.GetDepth();
+	const uint32 ParentItemID = static_cast<uint32>(Event.GetType() >> 32);
+	const uint32 ItemID = static_cast<uint32>(Event.GetType());
+
+	if (Depth == 0)
+	{
+		PopulateSubmissionBatchDetails(Builder, ItemID);
+		return;
+	}
+
+	const IRenderTraceProvider& Provider = Builder.GetProvider();
+	const FSubmissionBatch& Batch = Provider.GetSubmissionBatch(ParentItemID);
+	const FSubmissionEvent& SubmissionEvent = Batch.Events[ItemID];
+	switch (SubmissionEvent.Data.GetIndex())
+	{
+	case ESubmissionEvent_YieldSyncPoint:
+	{
+		const FSubmissionEvent::FYieldSyncPoint& YieldEvent = SubmissionEvent.Data.Get<FSubmissionEvent::FYieldSyncPoint>();
+		Builder.AddTitle(TEXTVIEW("Yield SyncPoint"));
+		Builder.AddNameValueTextLine(TEXTVIEW("Queue:"), GetPipeName(SubmissionEvent.Queue));
+		Builder.AddNameValueTextLine(TEXTVIEW("Pending payload:"), FString::Printf(TEXT("0x%llx"), YieldEvent.PayloadID));
+		Builder.AddNameValueTextLine(TEXTVIEW("Pending sync point:"), FString::Printf(TEXT("0x%llx"), YieldEvent.SyncPointID));
+		break;
+	}
+
+	case ESubmissionEvent_YieldManualFence:
+	{
+		const FSubmissionEvent::FYieldManualFence& YieldEvent = SubmissionEvent.Data.Get<FSubmissionEvent::FYieldManualFence>();
+		Builder.AddTitle(TEXTVIEW("Yield Manual Fence"));
+		Builder.AddNameValueTextLine(TEXTVIEW("Queue:"), GetPipeName(SubmissionEvent.Queue));
+		Builder.AddNameValueTextLine(TEXTVIEW("Pending payload:"), FString::Printf(TEXT("0x%llx"), YieldEvent.PayloadID));
+		Builder.AddNameValueTextLine(TEXTVIEW("Pending fence:"), FString::Printf(TEXT("0x%llx"), YieldEvent.FenceID));
+		Builder.AddNameValueTextLine(TEXTVIEW("Fence value:"), FString::Printf(TEXT("%llu"), YieldEvent.FenceValue));
+		break;
+	}
+
+	case ESubmissionEvent_WaitQueueFence:
+	{
+		const FSubmissionEvent::FWaitQueueFence& WaitQueueFenceEvent = SubmissionEvent.Data.Get<FSubmissionEvent::FWaitQueueFence>();
+		Builder.AddTitle(TEXTVIEW("Wait Queue Fence"));
+		Builder.AddNameValueTextLine(TEXTVIEW("Queue:"), GetPipeName(SubmissionEvent.Queue));
+		Builder.AddNameValueTextLine(TEXTVIEW("Other Queue:"), GetPipeName(WaitQueueFenceEvent.OtherQueue));
+		Builder.AddNameValueTextLine(TEXTVIEW("Fence Value:"), FString::Printf(TEXT("%llu"), WaitQueueFenceEvent.Value));
+		Builder.AddNameValueTextLine(TEXTVIEW("Payload:"), FString::Printf(TEXT("0x%llx"), WaitQueueFenceEvent.PayloadID));
+		break;
+	}
+
+	case ESubmissionEvent_WaitManualFence:
+	{
+		const FSubmissionEvent::FWaitManualFence& WaitManualFenceEvent = SubmissionEvent.Data.Get<FSubmissionEvent::FWaitManualFence>();
+		Builder.AddTitle(TEXTVIEW("Wait Manual Fence"));
+		Builder.AddNameValueTextLine(TEXTVIEW("Queue:"), GetPipeName(SubmissionEvent.Queue));
+		Builder.AddNameValueTextLine(TEXTVIEW("Fence ID:"), FString::Printf(TEXT("0x%llx"), WaitManualFenceEvent.FenceID));
+		Builder.AddNameValueTextLine(TEXTVIEW("Fence Value:"), FString::Printf(TEXT("%llu"), WaitManualFenceEvent.Value));
+		Builder.AddNameValueTextLine(TEXTVIEW("Payload:"), FString::Printf(TEXT("0x%llx"), WaitManualFenceEvent.PayloadID));
+		break;
+	}
+
+	case ESubmissionEvent_Execute:
+	{
+		const FSubmissionEvent::FExecute& ExecuteEvent = SubmissionEvent.Data.Get<FSubmissionEvent::FExecute>();
+		Builder.AddTitle(TEXTVIEW("Execute Command Lists"));
+		Builder.AddNameValueTextLine(TEXTVIEW("Queue:"), GetPipeName(SubmissionEvent.Queue));
+
+		Builder.AddNameValueTextLine(TEXTVIEW("Payloads:"), LexToString(ExecuteEvent.PayloadIDs.Num()));
+		{
+			FAppIDListBuilder<uint64> ListBuilder(Builder, ExecuteEvent.PayloadIDs.Num());
+			for (uint32 PayloadID : ExecuteEvent.PayloadIDs) { ListBuilder.AddItem(Provider.GetPlatformPayload(PayloadID).AppID); }
+			ListBuilder.Finish();
+		}
+
+		Builder.AddNameValueTextLine(TEXTVIEW("Command Lists:"), LexToString(ExecuteEvent.PlatformCmdListIDs.Num()));
+		{
+			FAppIDListBuilder<uint64> ListBuilder(Builder, ExecuteEvent.PlatformCmdListIDs.Num());
+			ListBuilder.AddIDs(ExecuteEvent.PlatformCmdListIDs, true);
+		}
+		break;
+	}
+
+	case ESubmissionEvent_SignalManualFence:
+	{
+		const FSubmissionEvent::FSignalManualFence& SignalManualFenceEvent = SubmissionEvent.Data.Get<FSubmissionEvent::FSignalManualFence>();
+		Builder.AddTitle(TEXTVIEW("Signal Manual Fence"));
+		Builder.AddNameValueTextLine(TEXTVIEW("Queue:"), GetPipeName(SubmissionEvent.Queue));
+		Builder.AddNameValueTextLine(TEXTVIEW("Fence ID:"), FString::Printf(TEXT("0x%llx"), SignalManualFenceEvent.FenceID));
+		Builder.AddNameValueTextLine(TEXTVIEW("Fence Value:"), FString::Printf(TEXT("%llu"), SignalManualFenceEvent.Value));
+		Builder.AddNameValueTextLine(TEXTVIEW("Payload:"), FString::Printf(TEXT("0x%llx"), SignalManualFenceEvent.PayloadID));
+		break;
+	}
+
+	case ESubmissionEvent_SignalQueueFence:
+	{
+		const FSubmissionEvent::FSignalQueueFence& SignalQueueFenceEvent = SubmissionEvent.Data.Get<FSubmissionEvent::FSignalQueueFence>();
+		Builder.AddTitle(TEXTVIEW("Signal Queue Fence"));
+		Builder.AddNameValueTextLine(TEXTVIEW("Queue:"), GetPipeName(SubmissionEvent.Queue));
+		Builder.AddNameValueTextLine(TEXTVIEW("Fence Value:"), FString::Printf(TEXT("%llu"), SignalQueueFenceEvent.Value));
+		Builder.AddNameValueTextLine(TEXTVIEW("Payload:"), FString::Printf(TEXT("0x%llx"), SignalQueueFenceEvent.PayloadID));
+		break;
+	}
+
+	case ESubmissionEvent_ResolveSyncPoint:
+	{
+		const FSubmissionEvent::FResolveSyncPoint& ResolveSyncPointEvent = SubmissionEvent.Data.Get<FSubmissionEvent::FResolveSyncPoint>();
+		Builder.AddTitle(TEXTVIEW("Resolve Sync Point"));
+		Builder.AddNameValueTextLine(TEXTVIEW("Queue:"), GetPipeName(SubmissionEvent.Queue));
+		Builder.AddNameValueTextLine(TEXTVIEW("Payload:"), FString::Printf(TEXT("0x%llx"), Provider.GetPlatformPayload(ResolveSyncPointEvent.PayloadID).AppID));
+		Builder.AddNameValueTextLine(TEXTVIEW("Sync Point:"), FString::Printf(TEXT("0x%llx"), Provider.GetSyncPoint(ResolveSyncPointEvent.SyncPointID).AppID));
+		Builder.AddNameValueTextLine(TEXTVIEW("Fence Value:"), FString::Printf(TEXT("%llu"), ResolveSyncPointEvent.Value));
+		break;
+	}
+	}
+}
+
+FInterruptTrack::FInterruptTrack(const FRenderTraceTimingViewSession& InSharedData, const FString& Name, const IRenderTraceProvider::TEventTimeline* InTimeline)
+	: FRenderTraceTrack(InSharedData, Name, InTimeline)
+{
+}
+
+void FInterruptTrack::OnRenderTraceEventSelected(const FTimingEvent& InSelectedEvent) const
+{
+	UE::Insights::Timing::ITimingViewSession* TimingView = SharedData.GetTimingViewSession();
+	const IRenderTraceProvider* Provider = SharedData.GetRenderTraceProvider();
+	if (!TimingView || !Provider)
+	{
+		return;
+	}
+
+	const uint32 Depth = InSelectedEvent.GetDepth();
+	const uint32 ParentItemID = static_cast<uint32>(InSelectedEvent.GetType() >> 32);
+	const uint32 ItemID = static_cast<uint32>(InSelectedEvent.GetType());
+
+	switch (Depth)
+	{
+	case 0:
+	{
+		// Wakeup event. Add an arrow to the CPU thread where batch was processed.
+		const FInterruptWakeUp& WakeUp = Provider->GetInterruptWakeUp(ItemID);
+
+		TSharedPtr<const FBaseTimingTrack> CpuThreadTrack = SharedData.GetCpuThreadTrack(WakeUp.ThreadID);
+		if (CpuThreadTrack)
+		{
+			TUniquePtr<ITimingEventRelation> Relation = MakeUnique<FRenderTraceRelation>(
+				SharedThis(this), WakeUp.StartTime, 0, CpuThreadTrack, WakeUp.StartTime, 0, GRelationColorExecThread
+			);
+			TimingView->AddRelation(Relation);
+		}
+
+		return;
+	}
+
+	case 1:
+	{
+		// Signal subevent. Add an arrow from the payload on the translate track.
+		const FInterruptWakeUp& WakeUp = Provider->GetInterruptWakeUp(ParentItemID);
+		const FInterruptFenceSignaledEvent& SignalEvent = WakeUp.SignalEvents[ItemID];
+
+		const FPlatformPayload& Payload = Provider->GetPlatformPayload(SignalEvent.PayloadID);
+		if (Payload.TranslateEvent.ParentItemID == INVALID_EVENT_ID)
+		{
+			return;
+		}
+
+		const FRHITranslateTask& TranslateTask = Provider->GetRHITranslateTask(Payload.TranslateEvent.ParentItemID);
+		check(TranslateTask.Type == ERHITranslateTaskType::Translate);
+
+		TSharedPtr<const FRHIThreadTrack> TranslateTrack = SharedData.GetRHIThreadTrack(TranslateTask.ThreadID);
+		if (TranslateTrack.IsValid())
+		{
+			const uint32 FirstContextDepth = 2;
+			const uint32 SourceDepth = FirstContextDepth + 2 * Payload.PipeIdx + 1;
+
+			TUniquePtr<ITimingEventRelation> Relation = MakeUnique<FRenderTraceRelation>(
+				TranslateTrack, Payload.EndTime, SourceDepth, SharedThis(this), InSelectedEvent.GetStartTime(), 1, GRelationColorInterrupt
+			);
+
+			TimingView->AddRelation(Relation);
+		}
+
+		return;
+	}
+	}
+}
+
+FString FInterruptTrack::GetEventName(const IRenderTraceProvider* Provider, uint32 InDepth, const FRenderTraceEvent& InEvent) const
+{
+	if (InDepth == 0)
+	{
+		return TEXT("ProcessInterruptQueue");
+	}
+
+	const FInterruptWakeUp& WakeUp = Provider->GetInterruptWakeUp(InEvent.ParentItemID);
+	const FInterruptFenceSignaledEvent& SignalEvent = WakeUp.SignalEvents[InEvent.ItemID];
+	return FString::Printf(TEXT("%s Signaled %llu"), GetPipeName(SignalEvent.Queue), SignalEvent.FenceValue);
+}
+
+void FInterruptTrack::PopulateEventDetails(IEventDetailsBuilder& Builder, const FTimingEvent& Event) const
+{
+	const IRenderTraceProvider& Provider = Builder.GetProvider();
+
+	const uint32 Depth = Event.GetDepth();
+	const uint32 ParentItemID = static_cast<uint32>(Event.GetType() >> 32);
+	const uint32 ItemID = static_cast<uint32>(Event.GetType());
+
+	switch (Depth)
+	{
+	case 0:
+	{
+		const FInterruptWakeUp& WakeUp = Provider.GetInterruptWakeUp(ItemID);
+		Builder.AddTitle(TEXTVIEW("ProcessInterruptQueue"));
+		Builder.AddNameValueTextLine(TEXTVIEW("Exit Status:"), GetExitStatusStr(WakeUp.ExitStatus));
+		Builder.AddNameValueTextLine(TEXTVIEW("Signal Events:"), LexToString(WakeUp.SignalEvents.Num()));
+
+		if (Builder.GetDetailLevel() == EEventDetailLevel::Full)
+		{
+			for (const FInterruptFenceSignaledEvent& SignalEvent : WakeUp.SignalEvents)
+			{
+				const FPlatformPayload& Payload = Provider.GetPlatformPayload(SignalEvent.PayloadID);
+				FString Label = FString::Printf(TEXT("[%s] Signaled:"), GetPipeName(SignalEvent.Queue));
+				FString Details = FString::Printf(TEXT("value=%llu payload=0x%llx lastCPU=%llu"), SignalEvent.FenceValue, Payload.AppID, SignalEvent.LastCPUSignaledFenceValue);
+				Builder.AddNameValueTextLine(Label, Details);
+			}
+		}
+
+		return;
+	}
+
+	case 1:
+	{
+		const FInterruptWakeUp& WakeUp = Provider.GetInterruptWakeUp(ParentItemID);
+		const FInterruptFenceSignaledEvent& SignalEvent = WakeUp.SignalEvents[ItemID];
+		const FPlatformPayload& Payload = Provider.GetPlatformPayload(SignalEvent.PayloadID);
+
+		Builder.AddTitle(TEXTVIEW("Signal Queue Fence"));
+		Builder.AddNameValueTextLine(TEXTVIEW("Queue:"), GetPipeName(SignalEvent.Queue));
+		Builder.AddNameValueTextLine(TEXTVIEW("Payload:"), FString::Printf(TEXT("0x%llx"), Payload.AppID));
+		Builder.AddNameValueTextLine(TEXTVIEW("Fence Value:"), LexToString(SignalEvent.FenceValue));
+		Builder.AddNameValueTextLine(TEXTVIEW("Last CPU Signaled Value:"), LexToString(SignalEvent.LastCPUSignaledFenceValue));
+		return;
+	}
+	}
+}
+
+} //namespace RenderTraceInsights
+} //namespace UE
diff --git a/Engine/Plugins/RenderTraceInsights/Source/RenderTraceInsights/Private/RenderTraceTrack.h b/Engine/Plugins/RenderTraceInsights/Source/RenderTraceInsights/Private/RenderTraceTrack.h
new file mode 100644
--- /dev/null
+++ b/Engine/Plugins/RenderTraceInsights/Source/RenderTraceInsights/Private/RenderTraceTrack.h
@@ -0,0 +1,193 @@
+// Copyright Epic Games, Inc. All Rights Reserved.
+
+#pragma once
+
+#include "Insights/ViewModels/TimingEventsTrack.h"
+#include "Insights/ViewModels/TimingTrackViewport.h"
+#include "Insights/ViewModels/TooltipDrawState.h"
+#include "Insights/ViewModels/ITimingEvent.h"
+#include "IRenderTraceProvider.h"
+
+class FMenuBuilder;
+class ITimingTrackDrawContext;
+class ITimingTrackUpdateContext;
+class FTimingTrackViewport;
+class FTimingEvent;
+
+namespace UE {
+namespace RenderTraceInsights {
+
+class FRenderTraceTimingViewSession;
+class IEventDetailsBuilder;
+
+class FRenderTraceRelation : public ITimingEventRelation
+{
+	INSIGHTS_DECLARE_RTTI(FRenderTraceRelation, ITimingEventRelation)
+public:
+	UE_NONCOPYABLE(FRenderTraceRelation)
+
+	FRenderTraceRelation(TSharedPtr<const FBaseTimingTrack> InSourceTrack, double InSourceTime, uint32 InSourceDepth, TSharedPtr<const FBaseTimingTrack> InTargetTrack, double InTargetTime, uint32 InTargetDepth, const FLinearColor& InColor);
+	virtual ~FRenderTraceRelation() = default;
+
+	void Draw(const UE::Insights::FDrawContext& DrawContext, const FTimingTrackViewport& Viewport, const ITimingViewDrawHelper& Helper, const ITimingEventRelation::EDrawFilter Filter) override;
+
+private:
+	TWeakPtr<const FBaseTimingTrack> SourceTrack;
+	TWeakPtr<const FBaseTimingTrack> TargetTrack;
+	double SourceTime;
+	double TargetTime;
+	uint32 SourceDepth;
+	uint32 TargetDepth;
+	FLinearColor Color;
+};
+
+class FRenderTraceTrack : public FTimingEventsTrack
+{
+	UE_NONCOPYABLE(FRenderTraceTrack)
+	INSIGHTS_DECLARE_RTTI(FRenderTraceTrack, FTimingEventsTrack)
+public:
+	FRenderTraceTrack(const FRenderTraceTimingViewSession& InSharedData, const FString& Name, const IRenderTraceProvider::TEventTimeline* InTimeline);
+
+	// FBaseTimingTrack interface.
+	const TSharedPtr<const ITimingEvent> GetEvent(double InTime, double SecondsPerPixel, int32 Depth) const override;
+	void InitTooltip(FTooltipDrawState& Tooltip, const ITimingEvent& TooltipEvent) const override;
+	void OnClipboardCopyEvent(const ITimingEvent& InSelectedEvent) const override;
+
+	// FTimingEventsTrack interface.
+	void BuildDrawState(ITimingEventsTrackDrawStateBuilder& Builder, const ITimingTrackUpdateContext& Context) override;
+
+	// Interface that subclasses must implement.
+	virtual void OnRenderTraceEventSelected(const FTimingEvent& InSelectedEvent) const {}
+	virtual FString GetEventName(const IRenderTraceProvider* Provider, uint32 InDepth, const FRenderTraceEvent& InEvent) const { return TEXT(""); }
+	virtual void PopulateEventDetails(IEventDetailsBuilder& Builder, const FTimingEvent& Event) const {}
+	virtual void CustomizeEventDrawState(ITimingEventsTrackDrawStateBuilder& Builder, double StartTime, double EndTime, uint32 Depth, const FRenderTraceEvent& Event) const {}
+	virtual TSharedPtr<FTimingEvent> GetEventAtTime(double EventTime, int32 EventDepth, double TimelineEventStart, double TimelineEventEnd, int32 TimelineEventDepth, const FRenderTraceEvent& TimelineEvent) const { return nullptr; }
+
+protected:
+	const FRenderTraceTimingViewSession& SharedData;
+	const IRenderTraceProvider::TEventTimeline* Timeline;
+};
+
+class FCommandListTrack : public FRenderTraceTrack
+{
+	UE_NONCOPYABLE(FCommandListTrack)
+	INSIGHTS_DECLARE_RTTI(FCommandListTrack, FRenderTraceTrack)
+public:
+	FCommandListTrack(const FRenderTraceTimingViewSession& InSharedData, const FString& Name, const IRenderTraceProvider::TEventTimeline* InTimeline);
+
+	// FRenderTraceTrack interface.
+	void OnRenderTraceEventSelected(const FTimingEvent& InSelectedEvent) const override;
+	FString GetEventName(const IRenderTraceProvider* Provider, uint32 InDepth, const FRenderTraceEvent& InEvent) const override;
+	void PopulateEventDetails(IEventDetailsBuilder& Builder, const FTimingEvent& Event) const override;
+
+private:
+	void OnCommandListSelected(const FTimingEvent& InSelectedEvent) const;
+	void OnRecordingEventSelected(const FTimingEvent& InSelectedEvent) const;
+	void PopulateCommandListEventDetails(IEventDetailsBuilder& Builder, uint32 CmdListID) const;
+	void PopulateRecordingEventDetails(IEventDetailsBuilder& Builder, uint32 PassID, double StartTime, double EndTime) const;
+};
+
+class FRDGThreadTrack : public FRenderTraceTrack
+{
+	UE_NONCOPYABLE(FRDGThreadTrack)
+	INSIGHTS_DECLARE_RTTI(FRDGThreadTrack, FRenderTraceTrack)
+public:
+	FRDGThreadTrack(const FRenderTraceTimingViewSession& InSharedData, const FString& Name, const IRenderTraceProvider::TEventTimeline* InTimeline);
+
+	// FRenderTraceTrack interface.
+	void OnRenderTraceEventSelected(const FTimingEvent& InSelectedEvent) const override;
+	FString GetEventName(const IRenderTraceProvider* Provider, uint32 InDepth, const FRenderTraceEvent& InEvent) const override;
+	void PopulateEventDetails(IEventDetailsBuilder& Builder, const FTimingEvent& Event) const override;
+};
+
+class FRenderThreadSubmissionTrack : public FRenderTraceTrack
+{
+	UE_NONCOPYABLE(FRenderThreadSubmissionTrack)
+	INSIGHTS_DECLARE_RTTI(FRenderThreadSubmissionTrack, FRenderTraceTrack)
+public:
+	FRenderThreadSubmissionTrack(const FRenderTraceTimingViewSession& InSharedData, const FString& Name, const IRenderTraceProvider::TEventTimeline* InTimeline);
+
+	// FRenderTraceTrack interface.
+	void OnRenderTraceEventSelected(const FTimingEvent& InSelectedEvent) const override;
+	FString GetEventName(const IRenderTraceProvider* Provider, uint32 InDepth, const FRenderTraceEvent& InEvent) const override;
+	void PopulateEventDetails(IEventDetailsBuilder& Builder, const FTimingEvent& Event) const override;
+};
+
+class FRHIThreadTrack : public FRenderTraceTrack
+{
+	UE_NONCOPYABLE(FRHIThreadTrack)
+	INSIGHTS_DECLARE_RTTI(FRHIThreadTrack, FRenderTraceTrack)
+public:
+	FRHIThreadTrack(const FRenderTraceTimingViewSession& InSharedData, const FString& Name, const IRenderTraceProvider::TEventTimeline* InTimeline);
+
+	// FRenderTraceTrack interface.
+	void OnRenderTraceEventSelected(const FTimingEvent& InSelectedEvent) const override;
+	FString GetEventName(const IRenderTraceProvider* Provider, uint32 InDepth, const FRenderTraceEvent& InEvent) const override;
+	void PopulateEventDetails(IEventDetailsBuilder& Builder, const FTimingEvent& Event) const override;
+	void CustomizeEventDrawState(ITimingEventsTrackDrawStateBuilder& Builder, double StartTime, double EndTime, uint32 Depth, const FRenderTraceEvent& Event) const override;
+	TSharedPtr<FTimingEvent> GetEventAtTime(double EventTime, int32 EventDepth, double TimelineEventStart, double TimelineEventEnd, int32 TimelineEventDepth, const FRenderTraceEvent& TimelineEvent) const override;
+
+private:
+	void CacheSplitFromTask(const IRenderTraceProvider& Provider, uint32 TaskID) const;
+	void CacheSplitToTask(const IRenderTraceProvider& Provider, uint32 TaskID) const;
+	void OnTranslateEventSelected(const FTimingEvent& InSelectedEvent) const;
+	void OnTranslateCommandListEventSelected(const FTimingEvent& InSelectedEvent) const;
+	void OnPayloadEventSelected(const FTimingEvent& InSelectedEvent, uint32 TranslateJobID, uint8 PipeIdx, uint32 PayloadIdx) const;
+	void PopulateTranslateJobDetails(IEventDetailsBuilder& Builder, uint32 TaskID, const FRHITranslateTask& Task) const;
+	void ListSyncPoints(IEventDetailsBuilder& Builder, const TCHAR* Title, TConstArrayView<uint32> SyncPointIDs) const;
+	void PopulatePayloadDetails(IEventDetailsBuilder& Builder, uint32 TranslateJobID, uint32 PipeIdx, uint32 PayloadIdx) const;
+	void PopulateTranslateContextDetails(IEventDetailsBuilder& Builder, uint32 TranslateJobID, uint32 PipeIdx) const;
+};
+
+class FRHISubmissionTrack : public FRenderTraceTrack
+{
+	UE_NONCOPYABLE(FRHISubmissionTrack)
+	INSIGHTS_DECLARE_RTTI(FRHISubmissionTrack, FRenderTraceTrack)
+public:
+	FRHISubmissionTrack(const FRenderTraceTimingViewSession& InSharedData, const FString& Name, const IRenderTraceProvider::TEventTimeline* InTimeline);
+
+	// FRenderTraceTrack interface.
+	void OnRenderTraceEventSelected(const FTimingEvent& InSelectedEvent) const override;
+	FString GetEventName(const IRenderTraceProvider* Provider, uint32 InDepth, const FRenderTraceEvent& InEvent) const override;
+	void PopulateEventDetails(IEventDetailsBuilder& Builder, const FTimingEvent& Event) const override;
+
+private:
+	void OnSubmitEventSelected(const FTimingEvent& InSelectedEvent) const;
+	void OnSubmitTranslateJobEventSelected(const FTimingEvent& InSelectedEvent) const;
+	void PopulateSubmitTaskDetails(IEventDetailsBuilder& Builder, uint32 TaskID, const FRHITranslateTask& Task) const;
+};
+
+class FSubmissionQueueTrack : public FRenderTraceTrack
+{
+	UE_NONCOPYABLE(FSubmissionQueueTrack)
+	INSIGHTS_DECLARE_RTTI(FSubmissionQueueTrack, FRenderTraceTrack)
+public:
+	FSubmissionQueueTrack(const FRenderTraceTimingViewSession& InSharedData, const FString& Name, const IRenderTraceProvider::TEventTimeline* InTimeline);
+
+	// FRenderTraceTrack interface.
+	void OnRenderTraceEventSelected(const FTimingEvent& InSelectedEvent) const override;
+	FString GetEventName(const IRenderTraceProvider* Provider, uint32 InDepth, const FRenderTraceEvent& InEvent) const override;
+	void PopulateEventDetails(IEventDetailsBuilder& Builder, const FTimingEvent& Event) const override;
+
+private:
+	void PopulateSubmissionBatchDetails(IEventDetailsBuilder& Builder, uint32 ItemID) const;
+	void OnSubmissionBatchSelected(const FTimingEvent& InSelectedEvent, uint32 BatchID) const;
+	void OnResolveSyncPointEventSelected(const FTimingEvent& InSelectedEvent, const FSubmissionEvent::FResolveSyncPoint& ResolveEvent) const;
+	void OnExecuteEventSelected(const FTimingEvent& InSelectedEvent, const FSubmissionEvent::FExecute& ExecEvent) const;
+};
+
+class FInterruptTrack : public FRenderTraceTrack
+{
+	UE_NONCOPYABLE(FInterruptTrack)
+	INSIGHTS_DECLARE_RTTI(FInterruptTrack, FRenderTraceTrack)
+public:
+	FInterruptTrack(const FRenderTraceTimingViewSession& InSharedData, const FString& Name, const IRenderTraceProvider::TEventTimeline* InTimeline);
+
+	// FRenderTraceTrack interface.
+	void OnRenderTraceEventSelected(const FTimingEvent& InSelectedEvent) const override;
+	FString GetEventName(const IRenderTraceProvider* Provider, uint32 InDepth, const FRenderTraceEvent& InEvent) const override;
+	void PopulateEventDetails(IEventDetailsBuilder& Builder, const FTimingEvent& Event) const override;
+};
+
+} // namespace RenderTraceInsights
+} //namespace UE
diff --git a/Engine/Plugins/RenderTraceInsights/Source/RenderTraceInsights/Public/IRenderTraceProvider.h b/Engine/Plugins/RenderTraceInsights/Source/RenderTraceInsights/Public/IRenderTraceProvider.h
new file mode 100644
--- /dev/null
+++ b/Engine/Plugins/RenderTraceInsights/Source/RenderTraceInsights/Public/IRenderTraceProvider.h
@@ -0,0 +1,473 @@
+// Copyright Epic Games, Inc. All Rights Reserved.
+
+#pragma once
+
+#include "HAL/Platform.h"
+#include "Containers/Array.h"
+#include "Misc/CoreMiscDefines.h"
+#include "Misc/TVariant.h"
+#include "TraceServices/Model/AnalysisSession.h"
+#include "Model/MonotonicTimeline.h"
+#include "Templates/Function.h"
+
+namespace TraceServices { class IAnalysisSession; }
+
+namespace UE
+{
+namespace RenderTraceInsights
+{
+
+static constexpr uint32 INVALID_EVENT_ID = ~0u;
+
+struct FTimeInterval
+{
+	FTimeInterval() : Start(0.0), End(0.0) {}
+	explicit FTimeInterval(double InStart) : Start(InStart), End(0.0) {}
+	FTimeInterval(double InStart, double InEnd) : Start(InStart), End(InEnd) {}
+	double Start;
+	double End;
+};
+
+struct FEventInterval
+{
+	FEventInterval(uint32 InItemID, double InStartTime) : ItemID(InItemID), Start(InStartTime), End(0.0) {}
+	FEventInterval(uint32 InItemID, double InStartTime, double InEndTime) : ItemID(InItemID), Start(InStartTime), End(InEndTime) {}
+	uint32 ItemID;
+	double Start;
+	double End;
+};
+
+struct FRenderTraceEvent
+{
+	FRenderTraceEvent() : ItemID(INVALID_EVENT_ID), ParentItemID(INVALID_EVENT_ID) {}
+	explicit FRenderTraceEvent(uint32 InItemID, uint32 InParentItemID = INVALID_EVENT_ID) : ItemID(InItemID), ParentItemID(InParentItemID) {}
+
+	uint32 ItemID;
+	uint32 ParentItemID;
+};
+
+enum class ECommandListType : uint8
+{
+	Regular,
+	Immediate,
+	Detached
+};
+
+enum ECommandListRecordingFlags
+{
+	ECommandListRecordingFlag_UsesLockFence = 1 << 0
+};
+
+enum ERHITranslateJobSplitFlags
+{
+	ERHITranslateJobSplitFlag_Parallel = 1 << 0,
+	ERHITranslateJobSplitFlag_Threshold = 1 << 1,
+	ERHITranslateJobSplitFlag_ParentChild = 1 << 2,
+	ERHITranslateJobSplitFlag_JumpThreads = 1 << 3
+};
+
+struct FCommandListInstance
+{
+	UE_NONCOPYABLE(FCommandListInstance)
+	FCommandListInstance(uint64 InAppID, double InCreateTime, ECommandListType InType)
+		: AppID(InAppID)
+		, Type(InType)
+		, CreateTime(InCreateTime)
+	{}
+
+	uint64 AppID;
+	ECommandListType Type;
+
+	// Important timestamps for the lifetime of the command list.
+	double CreateTime;
+	double FinishRecordingTime{};
+	double SubmitTime{};
+	uint32 SubmitPassID = INVALID_EVENT_ID;
+	double DestroyTime{};
+
+	// Recording information.
+	TArray<FEventInterval, TInlineAllocator<8>> RecordingEvents;
+	uint8 RecordingFlags = 0; // ECommandListRecordingFlags
+
+	// Detach information.
+	uint32 DetachedCmdListID = INVALID_EVENT_ID;
+	double DetachTime{};
+	uint32 SourceCmdListID = INVALID_EVENT_ID;
+
+	// Dispatch and translate information.
+	double DispatchTime{};
+	uint32 TranslateTaskID = INVALID_EVENT_ID;
+	uint64 TranslatePrevJobAppID = 0;
+	uint8 TranslateSplitFlags = 0; // ERHITranslateJobSplitFlags
+
+	// Timeline information.
+	double TimelineStart{};
+	int32 TimelineIndex = -1;
+};
+
+enum class ERDGPassType : uint8
+{
+	Regular,
+	Dispatch,
+	Submit
+};
+
+// Must be kept in sync with the enum of the same name in RenderGraphDefinitions.h.
+enum ERDGPassFlags
+{
+	ERDGPassFlag_Raster = 1 << 0,
+	ERDGPassFlag_Compute = 1 << 1,
+	ERDGPassFlag_AsyncCompute = 1 << 2,
+	ERDGPassFlag_Copy = 1 << 3,
+	ERDGPassFlag_NeverCull = 1 << 4,
+	ERDGPassFlag_SkipRenderPass = 1 << 5,
+	ERDGPassFlag_NeverMerge = 1 << 6,
+	ERDGPassFlag_NeverParallel = 1 << 7
+};
+
+// Must be kept in sync with the enum of the same name in RenderGraphPass.h.
+enum class ERDGPassTaskMode : uint8
+{
+	Inline = 0,
+	Await = 1,
+	Async = 2,
+	NotSet = 0xff
+};
+
+// Must be kept in sync ERHISubmitFlags. We don't want to use the enum in RHICommandList.h directly because some values are conditionally compiled.
+enum ERDGPassSubmitFlags
+{
+	ERDGPassSubmitFlag_SubmitToGPU = 1 << 0,
+	ERDGPassSubmitFlag_DeleteResources = 1 << 1,
+	ERDGPassSubmitFlag_FlushRHIThread = 1 << 2,
+	ERDGPassSubmitFlag_EndFrame = 1 << 3,
+	ERDGPassSubmitFlag_EnableBypass = 1 << 4,
+	ERDGPassSubmitFlag_DisableBypass = 1 << 5,
+	ERDGPassSubmitFlag_EnableDrawEvents = 1 << 6,
+	ERDGPassSubmitFlag_DisableDrawEvents = 1 << 7
+};
+
+struct FRDGPassInstance
+{
+	UE_NONCOPYABLE(FRDGPassInstance)
+	FRDGPassInstance(const FString& InName, uint32 InExecThreadID, double InStartTime, ERDGPassType InType)
+		: Name(InName)
+		, Type(InType)
+		, StartTime(InStartTime)
+		, ExecThreadID(InExecThreadID)
+	{}
+
+	FString Name;
+	ERDGPassType Type;
+	uint16 Flags = 0; // Either ERDGPassFlags or ERDGPassSubmitFlags.
+	ERDGPassTaskMode TaskMode = ERDGPassTaskMode::NotSet;
+
+	// Execution information.
+	double StartTime;
+	double EndTime{};
+	uint32 ExecThreadID;
+
+	// Recording command lists.
+	TArray<uint32, TInlineAllocator<8>> CommandListIDs;
+
+	// Extra data for render thread submit events.
+	uint32 SubmitTaskID = INVALID_EVENT_ID;
+};
+
+enum class ESyncPointType : uint8
+{
+	GPU = 0,
+	GPUAndCPU = 1,
+	Manual = 2
+};
+
+struct FSyncPoint
+{
+	UE_NONCOPYABLE(FSyncPoint)
+	FSyncPoint(uint64 InAppID, ESyncPointType InType) : AppID(InAppID), Type(InType) {}
+	uint64 AppID;
+	uint64 ResolvedValue = 0;
+	ESyncPointType Type;
+	uint32 ResolvedByPayload = INVALID_EVENT_ID;
+};
+
+struct FPlatformPayload
+{
+	UE_NONCOPYABLE(FPlatformPayload)
+	explicit FPlatformPayload(uint64 InAppID, double InStartTime, uint8 InPipeIdx) : AppID(InAppID), StartTime(InStartTime), PipeIdx(InPipeIdx) {}
+	uint64 AppID;
+	double StartTime;
+	double EndTime = 0.0;
+	uint8 PipeIdx;
+	TArray<uint32, TInlineAllocator<1>> WaitSyncPoints;
+	TArray<uint64, TInlineAllocator<1>> CmdLists;
+	TArray<uint32, TInlineAllocator<1>> SignalSyncPoints;
+	FRenderTraceEvent TranslateEvent;
+	FRenderTraceEvent ExecutionEvent;
+	FRenderTraceEvent InterruptEvent;
+	TArray<FRenderTraceEvent> ResolveSyncPointEvents;
+};
+
+struct FRHITranslateContext
+{
+	UE_NONCOPYABLE(FRHITranslateContext)
+	FRHITranslateContext() {}
+	uint64 RHIContextID = 0;
+	TArray<FTimeInterval, TInlineAllocator<4>> ActiveIntervals;
+	TArray<uint32, TInlineAllocator<2>> PayloadIDs;
+};
+
+enum class ERHITranslateTaskType : uint8
+{
+	Translate,
+	Submit
+};
+
+enum ERHITranslateJobFlags
+{
+	ERHITranslateJobFlag_Parallel = 1 << 0,
+	ERHITranslateJobFlag_UsingSubCmdLists = 1 << 1
+};
+
+struct FRHITranslateTask
+{
+	UE_NONCOPYABLE(FRHITranslateTask)
+
+	FRHITranslateTask(uint64 InAppID, ERHITranslateTaskType InType, uint32 InThreadID, double InStartTime)
+		: AppID(InAppID)
+		, Type(InType)
+		, ThreadID(InThreadID)
+		, StartTime(InStartTime)
+	{}
+
+	// Information about the task.
+	uint64 AppID;
+	ERHITranslateTaskType Type;
+	uint32 ThreadID;
+	double StartTime;
+	double EndTime{};
+
+	// Items processed by the job (command lists, translate jobs etc.).
+	TArray<FEventInterval, TInlineAllocator<8>> ProcessedItems;
+
+	// Next task in the translate chain (e.g. dispatch -> translate).
+	uint32 NextPhaseTaskID = INVALID_EVENT_ID;
+
+	uint8 JobFlags = 0; // ERHITranslateJobFlags
+
+	// The render-thread submit pass ID for submit tasks.
+	uint32 SubmitPassID = INVALID_EVENT_ID;
+	
+	// Previous and next jobs when a translate chain is split. This is not computed during analysis because it's too complicated.
+	// Instead, we compute and cache these values on demand when an event is selected.
+	mutable uint32 SplitFromTranslateTaskID = INVALID_EVENT_ID;
+	mutable uint32 SplitToTranslateTaskID = INVALID_EVENT_ID;
+
+	// Translation contexts, one per pipe: direct and async compute.
+	FRHITranslateContext Contexts[2];
+};
+
+// Must be kept in the same order as the subtype declarations in the FSubmissionEvent::Data variant.
+enum ESubmissionEvent
+{
+	ESubmissionEvent_YieldSyncPoint,
+	ESubmissionEvent_YieldManualFence,
+	ESubmissionEvent_WaitQueueFence,
+	ESubmissionEvent_WaitManualFence,
+	ESubmissionEvent_Execute,
+	ESubmissionEvent_SignalManualFence,
+	ESubmissionEvent_SignalQueueFence,
+	ESubmissionEvent_ResolveSyncPoint
+};
+
+struct FSubmissionEvent
+{
+	UE_NONCOPYABLE(FSubmissionEvent)
+	FSubmissionEvent(double InTimestamp, uint8 InQueue) : Timestamp(InTimestamp), Queue(InQueue) {}
+	~FSubmissionEvent() = default;
+
+	double Timestamp;
+	uint8 Queue;
+
+	struct FYieldSyncPoint
+	{
+		uint64 PayloadID;
+		uint64 SyncPointID;
+	};
+
+	struct FYieldManualFence
+	{
+		uint64 PayloadID;
+		uint64 FenceID;
+		uint64 FenceValue;
+	};
+
+	struct FWaitQueueFence
+	{
+		uint64 PayloadID;
+		uint8 OtherQueue;
+		uint64 Value;
+	};
+
+	struct FWaitManualFence
+	{
+		uint64 PayloadID;
+		uint64 FenceID;
+		uint64 Value;
+	};
+
+	struct FExecute
+	{
+		FExecute(TConstArrayView<uint32> InPayloadIDs, TConstArrayView<uint64> InPlatformCmdListIDs) : PayloadIDs(InPayloadIDs), PlatformCmdListIDs(InPlatformCmdListIDs) {}
+		TArray<uint32> PayloadIDs;
+		TArray<uint64> PlatformCmdListIDs;
+	};
+
+	struct FSignalManualFence
+	{
+		uint64 PayloadID;
+		uint64 FenceID;
+		uint64 Value;
+	};
+
+	struct FSignalQueueFence
+	{
+		uint64 PayloadID;
+		uint64 Value;
+	};
+
+	struct FResolveSyncPoint
+	{
+		uint32 PayloadID;
+		uint32 SyncPointID;
+		uint64 Value;
+	};
+
+	TVariant<FYieldSyncPoint, FYieldManualFence, FWaitQueueFence, FWaitManualFence, FExecute, FSignalManualFence, FSignalQueueFence, FResolveSyncPoint> Data;
+};
+
+struct FSubmissionBatch
+{
+	UE_NONCOPYABLE(FSubmissionBatch)
+	FSubmissionBatch(uint32 InThreadID, double InStartTime) : ThreadID(InThreadID), StartTime(InStartTime) {}
+
+	uint32 ThreadID;
+	double StartTime;
+	double EndTime = 0.0;
+	uint8 ExitStatus = 0;
+	TArray<FSubmissionEvent> Events;
+};
+
+struct FInterruptFenceSignaledEvent
+{
+	FInterruptFenceSignaledEvent(double InTimestamp, uint8 InQueue, uint32 InPayloadID, uint64 InFenceValue, uint64 InLastCPUSignaledFenceValue) :
+		Timestamp(InTimestamp), Queue(InQueue), PayloadID(InPayloadID), FenceValue(InFenceValue), LastCPUSignaledFenceValue(InLastCPUSignaledFenceValue) {}
+
+	double Timestamp;
+	uint8 Queue;
+	uint32 PayloadID;
+	uint64 FenceValue;
+	uint64 LastCPUSignaledFenceValue;
+};
+
+struct FInterruptWakeUp
+{
+	UE_NONCOPYABLE(FInterruptWakeUp)
+	FInterruptWakeUp(uint32 InThreadID, double InStartTime) : ThreadID(InThreadID), StartTime(InStartTime) {}
+
+	uint32 ThreadID;
+	double StartTime;
+	double EndTime = 0.0;
+	uint8 ExitStatus = 0;
+
+	TArray<FInterruptFenceSignaledEvent, TInlineAllocator<2>> SignalEvents;
+};
+
+class RENDERTRACEINSIGHTS_API IRenderTraceProvider : public TraceServices::IProvider
+{
+public:
+	using TEventTimeline = TraceServices::TMonotonicTimeline<FRenderTraceEvent>;
+
+	virtual uint64 GetTimelinesModCount() const = 0;
+
+	// Note: none of the accessor methods which take an item ID do explicit bounds checking. If an invalid ID is passed, the TArray
+	// operator[] will assert if checks are enabled, but otherwise garbage will be returned.
+
+	// RHI command list handling.
+	virtual int32 GetNumCommandLists() const = 0;
+	virtual const FCommandListInstance& GetCommandList(uint32 ID) const = 0;
+	virtual int32 GetNumCommandListTimelines() const = 0;
+	virtual void ReadCommandListTimelines(TFunctionRef<void(uint32, const TEventTimeline&)> Callback) const = 0;
+
+	// RDG pass handling.
+	virtual int32 GetNumRDGPasses() const = 0;
+	virtual const FRDGPassInstance& GetRDGPass(uint32 ID) const = 0;
+	virtual void ReadRDGTimelines(TFunctionRef<void(uint32, uint32, const TEventTimeline&)> Callback) const = 0;
+	virtual const TEventTimeline* GetRenderThreadSubmissionTimeline() const = 0;
+
+	// RHI translate task handling.
+	virtual int32 GetNumRHITranslateTasks() const = 0;
+	virtual const FRHITranslateTask& GetRHITranslateTask(uint32 ID) const = 0;
+	virtual uint32 FindRHITaskByPredicate(uint32 StartAtTaskID, double TimeRange, TFunctionRef<int(const FRHITranslateTask&)> Pred) const = 0;
+	virtual void EnumerateRHITranslateTimelines(TFunctionRef<void(uint32, uint32, const TEventTimeline&)> Callback) const = 0;
+	virtual const TEventTimeline* GetRHISubmissionTimeline() const = 0;
+
+	// Platform payload handling.
+	virtual int32 GetNumPlatformPayloads() const = 0;
+	virtual const FPlatformPayload& GetPlatformPayload(uint32 ID) const = 0;
+
+	// Sync point handling.
+	virtual int32 GetNumSyncPoints() const = 0;
+	virtual const FSyncPoint& GetSyncPoint(uint32 ID) const = 0;
+
+	// Submission queue handling.
+	virtual int32 GetNumSubmissionBatches() const = 0;
+	virtual const FSubmissionBatch& GetSubmissionBatch(uint32 ID) const = 0;
+	virtual const TEventTimeline* GetSubmissionQueueTimeline() const = 0;
+
+	// Interrupt queue handling.
+	virtual int32 GetNumInterruptWakeUps() const = 0;
+	virtual const FInterruptWakeUp& GetInterruptWakeUp(uint32 ID) const = 0;
+	virtual const TEventTimeline* GetInterruptTimeline() const = 0;
+};
+
+class RENDERTRACEINSIGHTS_API IEditableRenderTraceProvider : public IRenderTraceProvider, public TraceServices::IEditableProvider
+{
+public:
+	virtual TPair<uint32, FCommandListInstance&> AddCommandList(uint64 AppID, double Timestamp, ECommandListType Type) = 0;
+	virtual FCommandListInstance& EditCommandList(uint32 ID) = 0;
+	virtual void EnumerateCommandListsForEdit(TFunctionRef<void(uint32, FCommandListInstance&)> Callback) = 0;
+	virtual TPair<uint32, TEventTimeline&> AddCommandListTimeline() = 0;
+	virtual TEventTimeline& EditCommandListTimeline(uint32 Index) = 0;
+
+	virtual TPair<uint32, FRDGPassInstance&> AddRDGPass(const FString& Name, uint32 ThreadID, double Timestamp, ERDGPassType Type) = 0;
+	virtual FRDGPassInstance& EditRDGPass(uint32 ID) = 0;
+	virtual TEventTimeline& EditRDGThreadTimeline(uint32 ThreadID) = 0;
+	virtual TEventTimeline& EditRenderThreadSubmissionTimeline() = 0;
+
+	virtual TPair<uint32, FRHITranslateTask&> AddRHITranslateTask(uint64 AppID, ERHITranslateTaskType Type, uint32 ThreadID, double Timestamp) = 0;
+	virtual FRHITranslateTask& EditRHITranslateTask(uint32 ID) = 0;
+	virtual TEventTimeline& EditRHITranslateThreadTimeline(uint32 ThreadID) = 0;
+	virtual TEventTimeline& EditRHISubmissionTimeline() = 0;
+
+	virtual TPair<uint32, FPlatformPayload&> AddPlatformPayload(uint64 InAppID, double InStartTime, uint8 InPipeIdx) = 0;
+	virtual FPlatformPayload& EditPlatformPayload(uint32 ID) = 0;
+
+	virtual TPair<uint32, FSyncPoint&> AddSyncPoint(uint64 InAppID, ESyncPointType InType) = 0;
+	virtual FSyncPoint& EditSyncPoint(uint32 ID) = 0;
+
+	virtual TPair<uint32, FSubmissionBatch&> AddSubmissionBatch(uint32 ThreadID, double Timestamp) = 0;
+	virtual FSubmissionBatch& EditSubmissionBatch(uint32 ID) = 0;
+	virtual TEventTimeline& EditSubmissionQueueTimeline() = 0;
+
+	virtual TPair<uint32, FInterruptWakeUp&> AddInterruptWakeUp(uint32 ThreadID, double Timestamp) = 0;
+	virtual FInterruptWakeUp& EditInterruptWakeUp(uint32 ID) = 0;
+	virtual TEventTimeline& EditInterruptTimeline() = 0;
+};
+
+RENDERTRACEINSIGHTS_API FName GetRenderTraceProviderName();
+RENDERTRACEINSIGHTS_API const IRenderTraceProvider* ReadRenderTraceProvider(const TraceServices::IAnalysisSession& Session);
+RENDERTRACEINSIGHTS_API IEditableRenderTraceProvider* EditRenderTraceProvider(TraceServices::IAnalysisSession& Session);
+
+} //namespace RenderTraceInsights
+} //namespace UE
diff --git a/Engine/Plugins/RenderTraceInsights/Source/RenderTraceInsights/RenderTraceInsights.Build.cs b/Engine/Plugins/RenderTraceInsights/Source/RenderTraceInsights/RenderTraceInsights.Build.cs
new file mode 100644
--- /dev/null
+++ b/Engine/Plugins/RenderTraceInsights/Source/RenderTraceInsights/RenderTraceInsights.Build.cs
@@ -0,0 +1,28 @@
+// Copyright Epic Games, Inc. All Rights Reserved.
+
+namespace UnrealBuildTool.Rules
+{
+	public class RenderTraceInsights : ModuleRules
+	{
+		public RenderTraceInsights(ReadOnlyTargetRules Target) : base(Target)
+		{
+			PrivateDependencyModuleNames.AddRange(
+			new string[]
+			{
+				"ApplicationCore",
+				"Core",
+				"CoreUObject",
+				"RenderCore",
+				"RHI",
+				"SlateCore",
+				"Slate",
+				"TraceLog",
+				"TraceAnalysis",
+				"TraceServices",
+				"TraceInsights",
+				"TraceInsightsCore",
+				"InputCore"
+			});
+		}
+	}
+}
