Fix translate jobs being added to the timeline in incorrect order.

Jobs can be executed on one thread and finalized on another, so finalize events don't arrive in the same order as command list translation events. We need to make sure we add translate jobs in chronological order.

#rnx
#Synced-CL 51755722

diff --git a/Engine/Plugins/RenderTraceInsights/Source/RenderTraceInsights/Private/RenderTraceAnalyzer.cpp b/Engine/Plugins/RenderTraceInsights/Source/RenderTraceInsights/Private/RenderTraceAnalyzer.cpp
--- a/Engine/Plugins/RenderTraceInsights/Source/RenderTraceInsights/Private/RenderTraceAnalyzer.cpp
+++ b/Engine/Plugins/RenderTraceInsights/Source/RenderTraceInsights/Private/RenderTraceAnalyzer.cpp
@@ -294,6 +294,56 @@ void FRenderTraceAnalyzer::EndCommandList(uint32 PermanentCmdListID, double Time
 	AddCommandListToTimeline(PermanentCmdListID);
 }
 
+void FRenderTraceAnalyzer::AddTranslateTasksToTimeline(uint32 ThreadID)
+{
+	TArray<FLiveTranslateJob*, TInlineAllocator<8>> PendingJobs;
+	for (auto& It : LiveTranslateJobs)
+	{
+		const FRHITranslateTask& PendingTask = Provider.GetRHITranslateTask(It.Value.PermanentID);
+		if (!It.Value.bAddedToTimeline && PendingTask.ThreadID == ThreadID)
+		{
+			PendingJobs.Add(&It.Value);
+		}
+	}
+
+	if (PendingJobs.IsEmpty())
+	{
+		return;
+	}
+
+	// Sort the pending jobs by start time.
+	Algo::Sort(PendingJobs, [this](const FLiveTranslateJob* LiveJob1, const FLiveTranslateJob* LiveJob2) {
+		const FRHITranslateTask& Task1 = Provider.GetRHITranslateTask(LiveJob1->PermanentID);
+		const FRHITranslateTask& Task2 = Provider.GetRHITranslateTask(LiveJob2->PermanentID);
+		return Task1.StartTime < Task2.StartTime;
+	});
+
+	IRenderTraceProvider::TEventTimeline& Timeline = Provider.EditRHITranslateThreadTimeline(ThreadID);
+
+	// Add events until we find one that's still pending.
+	for (FLiveTranslateJob* LiveJob : PendingJobs)
+	{
+		const FRHITranslateTask& Task = Provider.GetRHITranslateTask(LiveJob->PermanentID);
+		if (Task.EndTime == 0.0)
+		{
+			break;
+		}
+
+		Timeline.AppendBeginEvent(Task.StartTime, FRenderTraceEvent(LiveJob->PermanentID));
+
+		// Add subevents for the command lists which are translated.
+		for (const FEventInterval& CmdList : Task.ProcessedItems)
+		{
+			Timeline.AppendBeginEvent(CmdList.Start, FRenderTraceEvent(CmdList.ItemID, LiveJob->PermanentID));
+			Timeline.AppendEndEvent(CmdList.End);
+		}
+
+		Timeline.AppendEndEvent(Task.EndTime);
+
+		LiveJob->bAddedToTimeline = true;
+	}
+}
+
 void FRenderTraceAnalyzer::EndTranslateTask(uint32 PermanentTaskID, double Timestamp)
 {
 	FRHITranslateTask& TranslateTask = Provider.EditRHITranslateTask(PermanentTaskID);
@@ -325,17 +375,10 @@ void FRenderTraceAnalyzer::EndTranslateTask(uint32 PermanentTaskID, double Times
 		}
 	}
 
-	IRenderTraceProvider::TEventTimeline& Timeline = Provider.EditRHITranslateThreadTimeline(TranslateTask.ThreadID);
-	Timeline.AppendBeginEvent(TranslateTask.StartTime, FRenderTraceEvent(PermanentTaskID));
-
-	// Add subevents for the command lists which are translated.
-	for (const FEventInterval& CmdList : TranslateTask.ProcessedItems)
-	{
-		Timeline.AppendBeginEvent(CmdList.Start, FRenderTraceEvent(CmdList.ItemID, PermanentTaskID));
-		Timeline.AppendEndEvent(CmdList.End);
-	}
-	// Translate tasks have well defined durations, use the timestamp.
-	Timeline.AppendEndEvent(TranslateTask.EndTime);
+	// We can't add the task to the timeline right away because jobs can be translated on one thread and finalized
+	// on another, which means finalize events can arrive in a different order than translate events. LiveTranslateJobs
+	// doubles as a queue of pending jobs, and we drain as much as possible from it each time this function is called.
+	AddTranslateTasksToTimeline(TranslateTask.ThreadID);
 }
 
 FPlatformPayload* FRenderTraceAnalyzer::FindPlatformPayload(uint64 PayloadID, const TCHAR* EventName, uint32 ThreadID)
@@ -614,10 +657,10 @@ void FRenderTraceAnalyzer::ProcessBeginTranslateCommandList(const FEventData& Ev
 
 	uint32 PermanentTranslateJobID = INVALID_EVENT_ID, ForceSplitFromJobID = INVALID_EVENT_ID;
 
-	const uint32* PermanentTranslateJobIDPtr = LiveTranslateJobs.Find(TranslateJobID);
-	if (PermanentTranslateJobIDPtr)
+	const FLiveTranslateJob* LiveTranslateJob = LiveTranslateJobs.Find(TranslateJobID);
+	if (LiveTranslateJob)
 	{
-		PermanentTranslateJobID = *PermanentTranslateJobIDPtr;
+		PermanentTranslateJobID = LiveTranslateJob->PermanentID;
 		FRHITranslateTask& TranslateJob = Provider.EditRHITranslateTask(PermanentTranslateJobID);
 		check(TranslateJob.Type == ERHITranslateTaskType::Translate);
 		if (TranslateJob.ThreadID != ThreadID)
@@ -634,7 +677,7 @@ void FRenderTraceAnalyzer::ProcessBeginTranslateCommandList(const FEventData& Ev
 	{
 		TPair<uint32, FRHITranslateTask&> NewTaskEntry = Provider.AddRHITranslateTask(TranslateJobID, ERHITranslateTaskType::Translate, ThreadID, Timestamp);
 		PermanentTranslateJobID = NewTaskEntry.Key;
-		LiveTranslateJobs.Add(TranslateJobID, PermanentTranslateJobID);
+		LiveTranslateJobs.Emplace(TranslateJobID, PermanentTranslateJobID);
 	}
 
 	CmdList.TranslateTaskID = PermanentTranslateJobID;
@@ -836,19 +879,20 @@ void FRenderTraceAnalyzer::ProcessEndTranslateCommandList(const FEventData& Even
 void FRenderTraceAnalyzer::ProcessFinalizeTranslateJob(const FEventData& EventData, uint32 ThreadID, double Timestamp)
 {
 	const uint64 TranslateJobID = EventData.GetValue<uint64>("TranslateJobID");
-	const uint32* PermanentTranslateJobID = LiveTranslateJobs.Find(TranslateJobID);
+	const FLiveTranslateJob* LiveTranslateJob = LiveTranslateJobs.Find(TranslateJobID);
 	const uint8 Flags = EventData.GetValue<uint8>("Flags");
-	if (!PermanentTranslateJobID)
+	if (!LiveTranslateJob)
 	{
 		UE_LOG(LogRenderTrace, Warning, TEXT("Got finalize event for unknown translate job ID %llx."), TranslateJobID);
 		return;
 	}
 
-	FRHITranslateTask& TranslateJob = Provider.EditRHITranslateTask(*PermanentTranslateJobID);
+	const uint32 PermanentTranslateJobID = LiveTranslateJob->PermanentID;
+	FRHITranslateTask& TranslateJob = Provider.EditRHITranslateTask(PermanentTranslateJobID);
 	check(TranslateJob.Type == ERHITranslateTaskType::Translate);
 	TranslateJob.JobFlags = Flags;
 
-	EndTranslateTask(*PermanentTranslateJobID, Timestamp);
+	EndTranslateTask(PermanentTranslateJobID, Timestamp);
 
 	for (int PipeIdx = 0; PipeIdx < 2; ++PipeIdx)
 	{
@@ -862,8 +906,7 @@ void FRenderTraceAnalyzer::ProcessFinalizeTranslateJob(const FEventData& EventDa
 		// notification (see the comment in ProcessSetTranslateContexts() for the full explanation). Only
 		// remove it from the live contexts map if we still own it.
 		const uint32* ExistingRHIContextMapping = LiveRHIContexts.Find(TranslateContext.RHIContextID);
-		check(ExistingRHIContextMapping);
-		if (*ExistingRHIContextMapping == *PermanentTranslateJobID)
+		if (ExistingRHIContextMapping && *ExistingRHIContextMapping == PermanentTranslateJobID)
 		{
 			LiveRHIContexts.Remove(TranslateContext.RHIContextID);
 		}
@@ -913,20 +956,35 @@ void FRenderTraceAnalyzer::ProcessSubmitTranslateJobs(const FEventData& EventDat
 
 	for (uint64 TranslateJobID : TranslateJobIDs)
 	{
-		const uint32* PermanentTranslateJobID = LiveTranslateJobs.Find(TranslateJobID);
-		if (!PermanentTranslateJobID)
+		const FLiveTranslateJob* LiveTranslateJob = LiveTranslateJobs.Find(TranslateJobID);
+		if (!LiveTranslateJob)
 		{
 			UE_LOG(LogRenderTrace, Warning, TEXT("Got submit event for unknown translate job ID %llx."), TranslateJobID);
 			continue;
 		}
 
-		FRHITranslateTask& TranslateJob = Provider.EditRHITranslateTask(*PermanentTranslateJobID);
+		const uint32 PermanentTranslateJobID = LiveTranslateJob->PermanentID;
+		FRHITranslateTask& TranslateJob = Provider.EditRHITranslateTask(PermanentTranslateJobID);
 		check(TranslateJob.Type == ERHITranslateTaskType::Translate);
 		TranslateJob.NextPhaseTaskID = NewTaskEntry.Key;
 
+		if (!LiveTranslateJob->bAddedToTimeline)
+		{
+			// This can happen early in the frame when the game thread is setting up the immediate command list and we don't get
+			// finalize events for the translate jobs.
+			if (TranslateJob.EndTime != 0.0)
+			{
+				AddTranslateTasksToTimeline(TranslateJob.ThreadID);
+			}
+			else
+			{
+				UE_LOG(LogRenderTrace, Warning, TEXT("Got submit event translate job ID %llx which doesn't have an end time."), TranslateJobID);
+			}
+		}
+
 		const double SubEventEndTime = FMath::Min(SubEventStartTime + SubEventDuration, SubmitTask.EndTime);
 
-		FEventInterval& SubEvent = SubmitTask.ProcessedItems.Emplace_GetRef(*PermanentTranslateJobID, SubEventStartTime, SubEventEndTime);
+		FEventInterval& SubEvent = SubmitTask.ProcessedItems.Emplace_GetRef(PermanentTranslateJobID, SubEventStartTime, SubEventEndTime);
 		Timeline.AppendBeginEvent(SubEvent.Start, FRenderTraceEvent(SubEvent.ItemID, NewTaskEntry.Key));
 		Timeline.AppendEndEvent(SubEvent.End);
 
@@ -1568,7 +1626,7 @@ void FRenderTraceAnalyzer::OnAnalysisEnd()
 		UE_LOG(LogRenderTrace, Warning, TEXT("Analysis ended with %u pending translate jobs, force-closing them."), LiveTranslateJobs.Num());
 		for (const auto& Pair : LiveTranslateJobs)
 		{
-			EndTranslateTask(Pair.Value, Timestamp);
+			EndTranslateTask(Pair.Value.PermanentID, Timestamp);
 		}
 		LiveTranslateJobs.Empty();
 	}
diff --git a/Engine/Plugins/RenderTraceInsights/Source/RenderTraceInsights/Private/RenderTraceAnalyzer.h b/Engine/Plugins/RenderTraceInsights/Source/RenderTraceInsights/Private/RenderTraceAnalyzer.h
--- a/Engine/Plugins/RenderTraceInsights/Source/RenderTraceInsights/Private/RenderTraceAnalyzer.h
+++ b/Engine/Plugins/RenderTraceInsights/Source/RenderTraceInsights/Private/RenderTraceAnalyzer.h
@@ -63,6 +63,13 @@ private:
 		Packet_D3D12ProcessInterruptQueueExit,
 	};
 
+	struct FLiveTranslateJob
+	{
+		explicit FLiveTranslateJob(uint32 InPermanentID) : PermanentID(InPermanentID) {}
+		uint32 PermanentID;
+		bool bAddedToTimeline = false;
+	};
+
 	TraceServices::IAnalysisSession& Session;
 	IEditableRenderTraceProvider& Provider;
 
@@ -71,7 +78,7 @@ private:
 	uint32 RHIThreadID = 0;
 	TMap<uint64, uint32> LiveRDGPasses;
 	TMap<uint64, uint32> LiveCommandLists;
-	TMap<uint64, uint32> LiveTranslateJobs;
+	TMap<uint64, FLiveTranslateJob> LiveTranslateJobs;
 	TMap<uint64, uint32> LiveRHIContexts;
 	TMap<uint64, uint64> PlatformToRHIContext;
 	TMap<uint32, uint32> PendingRTSubmitIDs;
@@ -87,6 +94,7 @@ private:
 	void EndRDGPassExecute(uint32 PermanentPassID, double Timestamp);
 	void AddCommandListToTimeline(uint32 PermanentCmdListID);
 	void EndCommandList(uint32 PermanentCmdListID, double Timestamp);
+	void AddTranslateTasksToTimeline(uint32 ThreadID);
 	void EndTranslateTask(uint32 PermanentTaskID, double Timestamp);
 	void EndPendingSubmissionBatch(double Timestamp, uint8 ExitStatus);
 	void EndPendingInterruptWakeUp(double Timestamp, uint8 ExitStatus);
